First pass at InteractiveToolsFramework module. EditorInteractiveToolsFramework is an implementation of the framework interfaces suitable for use inside an FEdMode.

diff --git a/Engine/Source/Editor/EditorInteractiveToolsFramework/EditorInteractiveToolsFramework.Build.cs b/Engine/Source/Editor/EditorInteractiveToolsFramework/EditorInteractiveToolsFramework.Build.cs
new file mode 100644
--- /dev/null
+++ b/Engine/Source/Editor/EditorInteractiveToolsFramework/EditorInteractiveToolsFramework.Build.cs
@@ -0,0 +1,58 @@
+// Copyright 1998-2019 Epic Games, Inc. All Rights Reserved.
+
+using UnrealBuildTool;
+
+public class EditorInteractiveToolsFramework : ModuleRules
+{
+	public EditorInteractiveToolsFramework(ReadOnlyTargetRules Target) : base(Target)
+	{
+		PCHUsage = ModuleRules.PCHUsageMode.UseExplicitOrSharedPCHs;
+		
+		PublicIncludePaths.AddRange(
+			new string[] {
+                //"ContentBrowser"
+				// ... add public include paths required here ...
+			}
+			);
+				
+		
+		PrivateIncludePaths.AddRange(
+			new string[] {
+				// ... add other private include paths required here ...
+			}
+			);
+			
+		
+		PublicDependencyModuleNames.AddRange(
+			new string[]
+			{
+				"Core"
+				// ... add other public dependencies that you statically link with here ...
+			}
+			);
+			
+		
+		PrivateDependencyModuleNames.AddRange(
+			new string[]
+			{
+				"CoreUObject",
+				"Engine",
+				"InputCore",
+				"UnrealEd",
+                "ContentBrowser",
+                "LevelEditor",
+                "InteractiveToolsFramework"
+
+				// ... add private dependencies that you statically link with here ...	
+			}
+            );
+		
+		
+		DynamicallyLoadedModuleNames.AddRange(
+			new string[]
+			{
+				// ... add any modules that your module loads dynamically here ...
+			}
+			);
+	}
+}
diff --git a/Engine/Source/Editor/EditorInteractiveToolsFramework/Private/EdModeInteractiveToolsContext.cpp b/Engine/Source/Editor/EditorInteractiveToolsFramework/Private/EdModeInteractiveToolsContext.cpp
new file mode 100644
--- /dev/null
+++ b/Engine/Source/Editor/EditorInteractiveToolsFramework/Private/EdModeInteractiveToolsContext.cpp
@@ -0,0 +1,384 @@
+// Copyright 1998-2019 Epic Games, Inc. All Rights Reserved.
+
+
+#include "EdModeInteractiveToolsContext.h"
+#include "EditorViewportClient.h"
+#include "EditorModeManager.h"
+#include "Engine/Selection.h"
+#include "Misc/ITransaction.h"
+#include "ScopedTransaction.h"
+
+#include "EditorToolAssetAPI.h"
+#include "EditorComponentSourceFactory.h"
+
+
+//#define ENABLE_DEBUG_PRINTING
+
+
+
+
+class FEdModeToolsContextQueriesImpl : public IToolsContextQueriesAPI
+{
+public:
+	UEdModeInteractiveToolsContext* ToolsContext;
+	FEdMode* EditorMode;
+
+	FEdModeToolsContextQueriesImpl(UEdModeInteractiveToolsContext* Context, FEdMode* EditorModeIn)
+	{
+		ToolsContext = Context;
+		EditorMode = EditorModeIn;
+	}
+
+	virtual void GetCurrentSelectionState(FToolBuilderState& StateOut) const override
+	{
+		StateOut.ToolManager = ToolsContext->ToolManager;
+		StateOut.World = EditorMode->GetWorld();
+		StateOut.SelectedActors = EditorMode->GetModeManager()->GetSelectedActors();
+		StateOut.SelectedComponents = EditorMode->GetModeManager()->GetSelectedComponents();
+		StateOut.SourceBuilder = ToolsContext->GetComponentSourceFactory();
+	}
+
+};
+
+
+
+
+class FEdModeToolsContextTransactionImpl : public IToolsContextTransactionsAPI
+{
+public:
+	UEdModeInteractiveToolsContext* ToolsContext;
+	FEdMode* EditorMode;
+
+	FEdModeToolsContextTransactionImpl(UEdModeInteractiveToolsContext* Context, FEdMode* EditorModeIn)
+	{
+		ToolsContext = Context;
+		EditorMode = EditorModeIn;
+	}
+
+
+	virtual void PostMessage(const TCHAR* Message, EToolMessageLevel Level) override
+	{
+		UE_LOG(LogTemp, Warning, TEXT("[ToolsContext] %s"), Message);
+	}
+
+
+	virtual void PostInvalidation() override
+	{
+		ToolsContext->PostInvalidation();
+	}
+
+	virtual void BeginUndoTransaction(const FText& Description) override
+	{
+		GEditor->BeginTransaction(Description);
+	}
+
+	virtual void EndUndoTransaction() override
+	{
+		GEditor->EndTransaction();
+	}
+
+	virtual void AppendChange(UObject* TargetObject, TUniquePtr<FChange> Change, const FText& Description) override
+	{
+		FScopedTransaction Transaction(Description);
+		check(GUndo != nullptr);
+		GUndo->StoreUndo(TargetObject, MoveTemp(Change));
+		// end transaction
+	}
+
+};
+
+
+
+
+
+UEdModeInteractiveToolsContext::UEdModeInteractiveToolsContext()
+{
+	EditorMode = nullptr;
+	QueriesAPI = nullptr;
+	TransactionAPI = nullptr;
+	AssetAPI = nullptr;
+	SourceFactory = nullptr;
+}
+
+
+
+void UEdModeInteractiveToolsContext::Initialize(IToolsContextQueriesAPI* queriesAPI, IToolsContextTransactionsAPI* transactionsAPI)
+{
+	UInteractiveToolsContext::Initialize(queriesAPI, transactionsAPI);
+
+	bInvalidationPending = false;
+}
+
+void UEdModeInteractiveToolsContext::Shutdown()
+{
+	UInteractiveToolsContext::Shutdown();
+}
+
+
+
+
+void UEdModeInteractiveToolsContext::InitializeContextFromEdMode(FEdMode* EditorModeIn)
+{
+	this->EditorMode = EditorModeIn;
+
+	this->TransactionAPI = new FEdModeToolsContextTransactionImpl(this, EditorModeIn);
+	this->QueriesAPI = new FEdModeToolsContextQueriesImpl(this, EditorModeIn);
+	this->AssetAPI = new FEditorToolAssetAPI();
+	this->SourceFactory = new FEditorComponentSourceFactory();
+
+	UInteractiveToolsContext::Initialize(QueriesAPI, TransactionAPI);
+
+	// enable auto invalidation in Editor, because invalidating for all hover and capture events is unpleasant
+	this->InputRouter->bAutoInvalidateOnHover = true;
+	this->InputRouter->bAutoInvalidateOnCapture = true;
+}
+
+
+void UEdModeInteractiveToolsContext::ShutdownContext()
+{
+	UInteractiveToolsContext::Shutdown();
+
+	if (QueriesAPI != nullptr)
+	{
+		delete QueriesAPI;
+		QueriesAPI = nullptr;
+	}
+
+	if (TransactionAPI != nullptr)
+	{
+		delete TransactionAPI;
+		TransactionAPI = nullptr;
+	}
+
+	if (AssetAPI != nullptr)
+	{
+		delete AssetAPI;
+		AssetAPI = nullptr;
+	}
+
+	if (SourceFactory != nullptr)
+	{
+		delete SourceFactory;
+		SourceFactory = nullptr;
+	}
+
+	this->EditorMode = nullptr;
+}
+
+
+
+void UEdModeInteractiveToolsContext::PostInvalidation()
+{
+	bInvalidationPending = true;
+}
+
+
+void UEdModeInteractiveToolsContext::Tick(FEditorViewportClient* ViewportClient, float DeltaTime)
+{
+	ToolManager->Tick(DeltaTime);
+
+	if (bInvalidationPending)
+	{
+		ViewportClient->Invalidate();
+		bInvalidationPending = false;
+	}
+}
+
+
+
+class TempRenderContext : public IToolsContextRenderAPI
+{
+public:
+	FPrimitiveDrawInterface* PDI;
+
+	virtual FPrimitiveDrawInterface* GetPrimitiveDrawInterface() override
+	{
+		return PDI;
+	}
+};
+
+
+void UEdModeInteractiveToolsContext::Render(const FSceneView* View, FViewport* Viewport, FPrimitiveDrawInterface* PDI)
+{
+	TempRenderContext RenderContext;
+	RenderContext.PDI = PDI;
+	ToolManager->Render(&RenderContext);
+}
+
+
+
+
+bool UEdModeInteractiveToolsContext::InputKey(FEditorViewportClient* ViewportClient, FViewport* Viewport, FKey Key, EInputEvent Event)
+{
+#ifdef ENABLE_DEBUG_PRINTING
+	if (Event == IE_Pressed) { UE_LOG(LogTemp, Warning, TEXT("PRESSED EVENT")); }
+	else if (Event == IE_Released) { UE_LOG(LogTemp, Warning, TEXT("RELEASED EVENT")); }
+	else if (Event == IE_Repeat) { UE_LOG(LogTemp, Warning, TEXT("REPEAT EVENT")); }
+	else if (Event == IE_Axis) { UE_LOG(LogTemp, Warning, TEXT("AXIS EVENT")); }
+	else if (Event == IE_DoubleClick) { UE_LOG(LogTemp, Warning, TEXT("DOUBLECLICK EVENT")); }
+#endif
+
+	bool bHandled = false;
+
+	if (Event == IE_Pressed || Event == IE_Released)
+	{
+		bool bIsLeftMouse = (Key == EKeys::LeftMouseButton);
+		bool bIsMiddleMouse = (Key == EKeys::MiddleMouseButton);
+		bool bIsRightMouse = (Key == EKeys::RightMouseButton);
+
+		if (bIsLeftMouse || bIsMiddleMouse || bIsRightMouse)
+		{
+
+			// early-out here if we are going to do camera manipulation
+			if (ViewportClient->IsAltPressed())
+			{
+				return bHandled;
+			}
+
+			FInputDeviceState InputState = CurrentInputState;
+			InputState.InputDevice = EInputDevices::Mouse;
+			InputState.SetKeyStates(
+				ViewportClient->IsShiftPressed(), ViewportClient->IsAltPressed(),
+				ViewportClient->IsCtrlPressed(), ViewportClient->IsCmdPressed());
+
+			if (bIsLeftMouse)
+			{
+				InputState.Mouse.Left.SetStates(
+					(Event == IE_Pressed), (Event == IE_Pressed), (Event == IE_Released));
+				CurrentInputState.Mouse.Left.bDown = (Event == IE_Pressed);
+			}
+			else if (bIsMiddleMouse)
+			{
+				InputState.Mouse.Middle.SetStates(
+					(Event == IE_Pressed), (Event == IE_Pressed), (Event == IE_Released));
+				CurrentInputState.Mouse.Middle.bDown = (Event == IE_Pressed);
+			}
+			else
+			{
+				InputState.Mouse.Right.SetStates(
+					(Event == IE_Pressed), (Event == IE_Pressed), (Event == IE_Released));
+				CurrentInputState.Mouse.Right.bDown = (Event == IE_Pressed);
+			}
+
+			InputRouter->PostInputEvent(InputState);
+
+			if (InputRouter->HasActiveMouseCapture())
+			{
+				// what is this about? MeshPaintMode has it...
+				ViewportClient->bLockFlightCamera = true;
+				bHandled = true;   // indicate that we handled this event,
+								   // which will disable camera movement/etc ?
+			}
+			else
+			{
+				//ViewportClient->bLockFlightCamera = false;
+			}
+
+		}
+
+	}
+
+	return bHandled;
+}
+
+
+
+
+bool UEdModeInteractiveToolsContext::MouseEnter(FEditorViewportClient* ViewportClient, FViewport* Viewport, int32 x, int32 y)
+{
+#ifdef ENABLE_DEBUG_PRINTING
+	UE_LOG(LogTemp, Warning, TEXT("MOUSE ENTER"));
+#endif
+
+	CurrentInputState.Mouse.Position2D = FVector2D(x, y);
+	CurrentInputState.Mouse.WorldRay = GetRayFromMousePos(ViewportClient, Viewport, x, y);
+
+	return false;
+}
+
+
+bool UEdModeInteractiveToolsContext::MouseMove(FEditorViewportClient* ViewportClient, FViewport* Viewport, int32 x, int32 y)
+{
+#ifdef ENABLE_DEBUG_PRINTING
+	//UE_LOG(LogTemp, Warning, TEXT("MOUSE MOVE"));
+#endif
+
+	CurrentInputState.Mouse.Position2D = FVector2D(x, y);
+	CurrentInputState.Mouse.WorldRay = GetRayFromMousePos(ViewportClient, Viewport, x, y);
+	FInputDeviceState InputState = CurrentInputState;
+	InputState.InputDevice = EInputDevices::Mouse;
+
+	if (InputRouter->HasActiveMouseCapture())
+		UE_LOG(LogTemp, Warning, TEXT("MOUSE MOVE DURING CAPTURE?"));
+
+	InputRouter->PostHoverInputEvent(InputState);
+
+	return false;
+}
+
+
+bool UEdModeInteractiveToolsContext::MouseLeave(FEditorViewportClient* ViewportClient, FViewport* Viewport)
+{
+#ifdef ENABLE_DEBUG_PRINTING
+	UE_LOG(LogTemp, Warning, TEXT("MOUSE LEAVE"));
+#endif
+
+	return false;
+}
+
+
+
+bool UEdModeInteractiveToolsContext::StartTracking(FEditorViewportClient* InViewportClient, FViewport* InViewport)
+{
+	return false;
+}
+
+bool UEdModeInteractiveToolsContext::CapturedMouseMove(FEditorViewportClient* InViewportClient, FViewport* InViewport, int32 InMouseX, int32 InMouseY)
+{
+#ifdef ENABLE_DEBUG_PRINTING
+	//UE_LOG(LogTemp, Warning, TEXT("CAPTURED MOUSE MOVE"));
+#endif
+
+	FVector2D OldPosition = CurrentInputState.Mouse.Position2D;
+	CurrentInputState.Mouse.Position2D = FVector2D(InMouseX, InMouseY);
+	CurrentInputState.Mouse.WorldRay = GetRayFromMousePos(InViewportClient, InViewport, InMouseX, InMouseY);
+
+	if (InputRouter->HasActiveMouseCapture())
+	{
+		FInputDeviceState InputState = CurrentInputState;
+		InputState.InputDevice = EInputDevices::Mouse;
+		InputState.Mouse.Delta2D = CurrentInputState.Mouse.Position2D - OldPosition;
+		InputRouter->PostInputEvent(InputState);
+		return true;
+	}
+
+	return false;
+}
+
+
+bool UEdModeInteractiveToolsContext::EndTracking(FEditorViewportClient* InViewportClient, FViewport* InViewport)
+{
+#ifdef ENABLE_DEBUG_PRINTING
+	UE_LOG(LogTemp, Warning, TEXT("END TRACKING"));
+#endif
+
+	return true;
+}
+
+
+
+
+
+
+FRay UEdModeInteractiveToolsContext::GetRayFromMousePos(FEditorViewportClient* ViewportClient, FViewport* Viewport, int MouseX, int MouseY)
+{
+	FSceneViewFamilyContext ViewFamily(FSceneViewFamily::ConstructionValues(
+		ViewportClient->Viewport,
+		ViewportClient->GetScene(),
+		ViewportClient->EngineShowFlags)
+		.SetRealtimeUpdate(ViewportClient->IsRealtime()));
+	FSceneView* View = ViewportClient->CalcSceneView(&ViewFamily);
+	FViewportCursorLocation MouseViewportRay(View, (FEditorViewportClient*)Viewport->GetClient(), MouseX, MouseY);
+
+	return FRay(MouseViewportRay.GetOrigin(), MouseViewportRay.GetDirection(), true);
+}
diff --git a/Engine/Source/Editor/EditorInteractiveToolsFramework/Private/EditorComponentSourceFactory.cpp b/Engine/Source/Editor/EditorInteractiveToolsFramework/Private/EditorComponentSourceFactory.cpp
new file mode 100644
--- /dev/null
+++ b/Engine/Source/Editor/EditorInteractiveToolsFramework/Private/EditorComponentSourceFactory.cpp
@@ -0,0 +1,95 @@
+// Copyright 1998-2019 Epic Games, Inc. All Rights Reserved.
+
+#include "EditorComponentSourceFactory.h"
+
+#include "Engine/StaticMesh.h"
+#include "Components/StaticMeshComponent.h"
+
+
+
+
+TUniquePtr<IMeshDescriptionSource> FEditorComponentSourceFactory::MakeMeshDescriptionSource(UObject* Component)
+{
+	UStaticMeshComponent* StaticMeshComp = Cast<UStaticMeshComponent>(Component);
+	if (StaticMeshComp != nullptr)
+	{
+		auto NewSource = new FStaticMeshComponentMeshDescriptionSource(StaticMeshComp, 0);
+		return TUniquePtr<IMeshDescriptionSource>(NewSource);
+	}
+	return nullptr;
+}
+
+
+
+
+
+FStaticMeshComponentMeshDescriptionSource::FStaticMeshComponentMeshDescriptionSource(
+	UStaticMeshComponent* ComponentIn, int LODIndex)
+{
+	this->Component = ComponentIn;
+	this->LODIndex = LODIndex;
+}
+
+
+AActor* FStaticMeshComponentMeshDescriptionSource::GetOwnerActor() const
+{
+	return Component->GetOwner();
+}
+
+UObject* FStaticMeshComponentMeshDescriptionSource::GetOwnerComponent() const
+{
+	return Component;
+}
+
+
+void FStaticMeshComponentMeshDescriptionSource::SetOwnerVisibility(bool bVisible) const
+{
+	Component->SetVisibility(bVisible);
+}
+
+FMeshDescription* FStaticMeshComponentMeshDescriptionSource::GetMeshDescription() const
+{
+	return Component->GetStaticMesh()->GetMeshDescription(0);
+}
+
+UMaterialInterface* FStaticMeshComponentMeshDescriptionSource::GetMaterial(int32 MaterialIndex) const
+{
+	return Component->GetMaterial(MaterialIndex);
+}
+
+FTransform FStaticMeshComponentMeshDescriptionSource::GetWorldTransform() const
+{
+	return Component->GetOwner()->GetActorTransform();
+}
+
+
+
+bool FStaticMeshComponentMeshDescriptionSource::IsReadOnly() const
+{
+	return false;
+}
+
+
+void FStaticMeshComponentMeshDescriptionSource::CommitInPlaceModification(const TFunction<void(FMeshDescription*)>& ModifyFunction) 
+{
+	UStaticMesh* StaticMesh = Component->GetStaticMesh();
+	StaticMesh->Modify();
+	FMeshDescription* MeshDescription = StaticMesh->GetMeshDescription(LODIndex);
+
+	ModifyFunction(MeshDescription);
+
+	StaticMesh->CommitMeshDescription(LODIndex);
+	StaticMesh->PostEditChange();
+}
+
+
+bool FStaticMeshComponentMeshDescriptionSource::HitTest(const FRay& WorldRay, FHitResult& OutHit) const
+{
+	FVector End = WorldRay.PointAt(HALF_WORLD_MAX);
+	if (Component->LineTraceComponent(OutHit, WorldRay.Origin, End, FCollisionQueryParams(SCENE_QUERY_STAT(HitTest), true)))
+	{
+		return true;
+	}
+
+	return false;
+}
\ No newline at end of file
diff --git a/Engine/Source/Editor/EditorInteractiveToolsFramework/Private/EditorInteractiveToolsFrameworkModule.cpp b/Engine/Source/Editor/EditorInteractiveToolsFramework/Private/EditorInteractiveToolsFrameworkModule.cpp
new file mode 100644
--- /dev/null
+++ b/Engine/Source/Editor/EditorInteractiveToolsFramework/Private/EditorInteractiveToolsFrameworkModule.cpp
@@ -0,0 +1,20 @@
+// Copyright 1998-2018 Epic Games, Inc. All Rights Reserved.
+
+#include "EditorInteractiveToolsFrameworkModule.h"
+
+#define LOCTEXT_NAMESPACE "FEditorInteractiveToolsFrameworkModule"
+
+void FEditorInteractiveToolsFrameworkModule::StartupModule()
+{
+	// This code will execute after your module is loaded into memory; the exact timing is specified in the .uplugin file per-module
+}
+
+void FEditorInteractiveToolsFrameworkModule::ShutdownModule()
+{
+	// This function may be called during shutdown to clean up your module.  For modules that support dynamic reloading,
+	// we call this function before unloading the module.
+}
+
+#undef LOCTEXT_NAMESPACE
+	
+IMPLEMENT_MODULE(FEditorInteractiveToolsFrameworkModule, EditorInteractiveToolsFramework)
\ No newline at end of file
diff --git a/Engine/Source/Editor/EditorInteractiveToolsFramework/Private/EditorToolAssetAPI.cpp b/Engine/Source/Editor/EditorInteractiveToolsFramework/Private/EditorToolAssetAPI.cpp
new file mode 100644
--- /dev/null
+++ b/Engine/Source/Editor/EditorInteractiveToolsFramework/Private/EditorToolAssetAPI.cpp
@@ -0,0 +1,93 @@
+// Copyright 1998-2019 Epic Games, Inc. All Rights Reserved.
+
+#include "EditorToolAssetAPI.h"
+#include "CoreMinimal.h"
+#include "GameFramework/Actor.h"
+
+
+#include "AssetRegistryModule.h"
+#include "FileHelpers.h"
+#include "PackageTools.h"
+#include "ObjectTools.h"
+
+// for content-browser things
+#include "ContentBrowserModule.h"
+#include "IContentBrowserSingleton.h"
+
+
+FString FEditorToolAssetAPI::GetActiveAssetFolderPath()
+{
+	check(false);  	// [RMS] this only works if we have an asset selected!
+
+	IContentBrowserSingleton& ContentBrowser = FModuleManager::LoadModuleChecked<FContentBrowserModule>("ContentBrowser").Get();
+
+	//ContentBrowser.CreatePathPicker()
+
+	TArray<FAssetData> SelectedAssets;
+	ContentBrowser.GetSelectedAssets(SelectedAssets);
+	return SelectedAssets[0].PackagePath.ToString();
+}
+
+
+
+FString FEditorToolAssetAPI::MakePackageName(const FString& AssetName, const FString& FolderPath)
+{
+	return FolderPath + TEXT("/") + AssetName;
+}
+
+
+FString FEditorToolAssetAPI::MakeUniqueAssetName(const FString& AssetName, const FString& FolderPath)
+{
+	FString NewPackageName = MakePackageName(AssetName, FolderPath);
+	if (FPackageName::DoesPackageExist(NewPackageName) == false)
+	{
+		return AssetName;
+	}
+
+	int Counter = 1;
+	FString UniqueName = AssetName;
+	while ( FPackageName::DoesPackageExist(NewPackageName) )
+	{
+		UniqueName = AssetName + FString::Printf(TEXT("_%d"), Counter++);
+		NewPackageName = MakePackageName(UniqueName, FolderPath);
+	}
+	return UniqueName;
+}
+
+
+UPackage* FEditorToolAssetAPI::CreateNewPackage(const FString& AssetName, const FString& FolderPath)
+{
+	FString NewPackageName = MakePackageName(AssetName, FolderPath);
+	NewPackageName = UPackageTools::SanitizePackageName(NewPackageName);
+	UPackage* AssetPackage = CreatePackage(nullptr, *NewPackageName);
+	return AssetPackage;
+}
+
+
+
+void FEditorToolAssetAPI::InteractiveSaveGeneratedAsset(UObject* Asset, UPackage* AssetPackage)
+{
+	Asset->MarkPackageDirty();
+
+	FAssetRegistryModule::AssetCreated(Asset);
+
+	TArray<UPackage*> PackagesToSave;
+	PackagesToSave.Add(AssetPackage);
+	bool bCheckDirty = true;
+	bool bPromptToSave = true;
+	FEditorFileUtils::PromptForCheckoutAndSave(PackagesToSave, bCheckDirty, bPromptToSave);
+}
+
+
+void FEditorToolAssetAPI::AutoSaveGeneratedAsset(UObject* Asset, UPackage* AssetPackage)
+{
+	Asset->MarkPackageDirty();
+
+	FAssetRegistryModule::AssetCreated(Asset);
+
+	TArray<UPackage*> PackagesToSave;
+	PackagesToSave.Add(AssetPackage);
+	bool bCheckDirty = true;
+	bool bPromptToSave = false;
+	FEditorFileUtils::PromptForCheckoutAndSave(PackagesToSave, bCheckDirty, bPromptToSave);
+}
\ No newline at end of file
diff --git a/Engine/Source/Editor/EditorInteractiveToolsFramework/Public/EdModeInteractiveToolsContext.h b/Engine/Source/Editor/EditorInteractiveToolsFramework/Public/EdModeInteractiveToolsContext.h
new file mode 100644
--- /dev/null
+++ b/Engine/Source/Editor/EditorInteractiveToolsFramework/Public/EdModeInteractiveToolsContext.h
@@ -0,0 +1,74 @@
+// Copyright 1998-2018 Epic Games, Inc. All Rights Reserved.
+
+#pragma once
+
+#include "CoreMinimal.h"
+#include "EdMode.h"
+#include "InteractiveToolsContext.h"
+#include "EdModeInteractiveToolsContext.generated.h"
+
+/**
+ * EdModeInteractiveToolsContext is an extension/adapter of an InteractiveToolsContext which 
+ * allows it to be easily embedded inside an FEdMode. A set of functions are provided which can be
+ * called from the FEdMode functions of the same name. These will handle the data type
+ * conversions and forwarding calls necessary to operate the ToolsContext
+ */
+UCLASS(Transient)
+class EDITORINTERACTIVETOOLSFRAMEWORK_API UEdModeInteractiveToolsContext : public UInteractiveToolsContext
+{
+	GENERATED_BODY()
+
+public:
+	UEdModeInteractiveToolsContext();
+
+
+	virtual void InitializeContextFromEdMode(FEdMode* EditorMode);
+	virtual void ShutdownContext();
+
+
+	IToolsContextQueriesAPI* GetQueriesAPI() const { return QueriesAPI; }
+	IToolsContextTransactionsAPI* GetTransactionAPI() const { return TransactionAPI; }
+	IToolsContextAssetAPI* GetAssetAPI() const { return AssetAPI; }
+	IComponentSourceFactory* GetComponentSourceFactory() const { return SourceFactory; }
+
+	virtual void PostInvalidation();
+
+
+	// call these from your FEdMode functions of the same name
+
+	virtual void Tick(FEditorViewportClient* ViewportClient, float DeltaTime);
+	virtual void Render(const FSceneView* View, FViewport* Viewport, FPrimitiveDrawInterface* PDI);
+
+	virtual bool InputKey(FEditorViewportClient* ViewportClient, FViewport* Viewport, FKey Key, EInputEvent Event);
+
+	virtual bool MouseEnter(FEditorViewportClient* ViewportClient, FViewport* Viewport, int32 x, int32 y);
+	virtual bool MouseLeave(FEditorViewportClient* ViewportClient, FViewport* Viewport);
+	virtual bool MouseMove(FEditorViewportClient* ViewportClient, FViewport* Viewport, int32 x, int32 y);
+
+	virtual bool StartTracking(FEditorViewportClient* InViewportClient, FViewport* InViewport);
+	virtual bool CapturedMouseMove(FEditorViewportClient* InViewportClient, FViewport* InViewport, int32 InMouseX, int32 InMouseY);
+	virtual bool EndTracking(FEditorViewportClient* InViewportClient, FViewport* InViewport);
+
+
+protected:
+	// we hide these 
+	virtual void Initialize(IToolsContextQueriesAPI* QueriesAPI, IToolsContextTransactionsAPI* TransactionsAPI) override;
+	virtual void Shutdown() override;
+
+
+protected:
+	FEdMode* EditorMode;
+
+	IToolsContextQueriesAPI* QueriesAPI;
+	IToolsContextTransactionsAPI* TransactionAPI;
+	IToolsContextAssetAPI* AssetAPI;
+	IComponentSourceFactory* SourceFactory;
+
+	bool bInvalidationPending;
+
+	/** Input event instance used to keep track of various button states, etc, that we cannot directly query on-demand */
+	FInputDeviceState CurrentInputState;
+
+	static FRay GetRayFromMousePos(FEditorViewportClient* ViewportClient, FViewport* Viewport, int MouseX, int MouseY);
+
+};
diff --git a/Engine/Source/Editor/EditorInteractiveToolsFramework/Public/EditorComponentSourceFactory.h b/Engine/Source/Editor/EditorInteractiveToolsFramework/Public/EditorComponentSourceFactory.h
new file mode 100644
--- /dev/null
+++ b/Engine/Source/Editor/EditorInteractiveToolsFramework/Public/EditorComponentSourceFactory.h
@@ -0,0 +1,54 @@
+// Copyright 1998-2019 Epic Games, Inc. All Rights Reserved.
+
+#pragma once
+
+#include "CoreMinimal.h"
+#include "ComponentSourceInterfaces.h"
+
+
+
+
+/**
+ * Component Source Factory for use in the Editor (ie inside FEdMode, etc)
+ */
+class EDITORINTERACTIVETOOLSFRAMEWORK_API FEditorComponentSourceFactory : public IComponentSourceFactory
+{
+public:
+	virtual ~FEditorComponentSourceFactory() {}
+
+	virtual TUniquePtr<IMeshDescriptionSource> MakeMeshDescriptionSource(UObject* Component);
+};
+
+
+
+
+
+class UStaticMeshComponent;
+
+/**
+ * This MeshDescriptionSource provides a specific LOD from a StaticMeshComponent
+ */
+class EDITORINTERACTIVETOOLSFRAMEWORK_API FStaticMeshComponentMeshDescriptionSource : public IMeshDescriptionSource
+{
+public:
+	UStaticMeshComponent* Component;
+	int LODIndex;
+
+	FStaticMeshComponentMeshDescriptionSource(
+		UStaticMeshComponent* ComponentIn,
+		int LODIndex = 0);
+	virtual ~FStaticMeshComponentMeshDescriptionSource() {}
+
+	virtual AActor* GetOwnerActor() const override;
+	virtual UObject* GetOwnerComponent() const override;
+	virtual FMeshDescription* GetMeshDescription() const override;
+	virtual UMaterialInterface* GetMaterial(int32 MaterialIndex) const override;
+	virtual FTransform GetWorldTransform() const override;
+	virtual bool HitTest(const FRay& WorldRay, FHitResult& OutHit) const override;
+
+	virtual bool IsReadOnly() const override;
+
+	virtual void SetOwnerVisibility(bool bVisible) const override;
+	virtual void CommitInPlaceModification(const TFunction<void(FMeshDescription*)>& ModifyFunction) override;
+};
+
diff --git a/Engine/Source/Editor/EditorInteractiveToolsFramework/Public/EditorInteractiveToolsFrameworkModule.h b/Engine/Source/Editor/EditorInteractiveToolsFramework/Public/EditorInteractiveToolsFrameworkModule.h
new file mode 100644
--- /dev/null
+++ b/Engine/Source/Editor/EditorInteractiveToolsFramework/Public/EditorInteractiveToolsFrameworkModule.h
@@ -0,0 +1,15 @@
+// Copyright 1998-2018 Epic Games, Inc. All Rights Reserved.
+
+#pragma once
+
+#include "CoreMinimal.h"
+#include "Modules/ModuleManager.h"
+
+class FEditorInteractiveToolsFrameworkModule : public IModuleInterface
+{
+public:
+
+	/** IModuleInterface implementation */
+	virtual void StartupModule() override;
+	virtual void ShutdownModule() override;
+};
diff --git a/Engine/Source/Editor/EditorInteractiveToolsFramework/Public/EditorToolAssetAPI.h b/Engine/Source/Editor/EditorInteractiveToolsFramework/Public/EditorToolAssetAPI.h
new file mode 100644
--- /dev/null
+++ b/Engine/Source/Editor/EditorInteractiveToolsFramework/Public/EditorToolAssetAPI.h
@@ -0,0 +1,28 @@
+// Copyright 1998-2019 Epic Games, Inc. All Rights Reserved.
+
+#pragma once
+
+#include "CoreMinimal.h"
+#include "ToolContextInterfaces.h"
+
+/**
+ * Implementation of ToolsContext Asset management API that is suitable for use
+ * inside the Editor (eg inside an FEdMode)
+ */
+class EDITORINTERACTIVETOOLSFRAMEWORK_API FEditorToolAssetAPI : public IToolsContextAssetAPI
+{
+public:
+	virtual FString GetActiveAssetFolderPath() override;
+
+	virtual FString MakePackageName(const FString& AssetName, const FString& FolderPath) override;
+
+	virtual FString MakeUniqueAssetName(const FString& AssetName, const FString& FolderPath) override;
+
+	// FolderPath should neither start with, nor end with, path separators
+	// Generated path will be /Game/FolderPath/AssetName
+	virtual UPackage* CreateNewPackage(const FString& AssetName, const FString& FolderPath) override;
+
+	virtual void InteractiveSaveGeneratedAsset(UObject* Asset, UPackage* AssetPackage) override;
+
+	virtual void AutoSaveGeneratedAsset(UObject* Asset, UPackage* AssetPackage) override;
+};
\ No newline at end of file
diff --git a/Engine/Source/Runtime/InteractiveToolsFramework/InteractiveToolsFramework.Build.cs b/Engine/Source/Runtime/InteractiveToolsFramework/InteractiveToolsFramework.Build.cs
new file mode 100644
--- /dev/null
+++ b/Engine/Source/Runtime/InteractiveToolsFramework/InteractiveToolsFramework.Build.cs
@@ -0,0 +1,53 @@
+// Copyright 1998-2018 Epic Games, Inc. All Rights Reserved.
+
+using UnrealBuildTool;
+
+public class InteractiveToolsFramework : ModuleRules
+{
+	public InteractiveToolsFramework(ReadOnlyTargetRules Target) : base(Target)
+	{
+		PCHUsage = ModuleRules.PCHUsageMode.UseExplicitOrSharedPCHs;
+		
+		PublicIncludePaths.AddRange(
+			new string[] {
+				// ... add public include paths required here ...
+			}
+			);
+				
+		
+		PrivateIncludePaths.AddRange(
+			new string[] {
+				// ... add other private include paths required here ...
+			}
+			);
+			
+		
+		PublicDependencyModuleNames.AddRange(
+			new string[]
+			{
+				"Core"
+				// ... add other public dependencies that you statically link with here ...
+			}
+            );
+			
+		
+		PrivateDependencyModuleNames.AddRange(
+			new string[]
+			{
+				"CoreUObject",
+				"Engine"
+				//"Slate",
+				//"SlateCore",
+				// ... add private dependencies that you statically link with here ...	
+			}
+			);
+		
+		
+		DynamicallyLoadedModuleNames.AddRange(
+			new string[]
+			{
+				// ... add any modules that your module loads dynamically here ...
+			}
+			);
+	}
+}
diff --git a/Engine/Source/Runtime/InteractiveToolsFramework/Private/InputBehavior.cpp b/Engine/Source/Runtime/InteractiveToolsFramework/Private/InputBehavior.cpp
new file mode 100644
--- /dev/null
+++ b/Engine/Source/Runtime/InteractiveToolsFramework/Private/InputBehavior.cpp
@@ -0,0 +1,72 @@
+// Copyright 1998-2019 Epic Games, Inc. All Rights Reserved.
+
+
+#include "InputBehavior.h"
+
+
+bool operator<(const FInputCaptureRequest& l, const FInputCaptureRequest& r)
+{
+	if (l.Source->GetPriority() == r.Source->GetPriority())
+	{
+		return l.HitDepth < r.HitDepth;
+	}
+	else
+	{
+		return l.Source->GetPriority() < r.Source->GetPriority();
+	}
+}
+
+
+
+UInputBehavior::UInputBehavior()
+{
+	DefaultPriority = FInputCapturePriority(100);
+}
+
+FInputCapturePriority UInputBehavior::GetPriority()
+{
+	return DefaultPriority;
+}
+
+EInputDevices UInputBehavior::GetSupportedDevices()
+{
+	return EInputDevices::Mouse;
+}
+
+
+FInputCaptureRequest UInputBehavior::WantsCapture(const FInputDeviceState& input)
+{
+	return FInputCaptureRequest::Ignore();
+}
+
+FInputCaptureUpdate UInputBehavior::BeginCapture(const FInputDeviceState& input, EInputCaptureSide eSide)
+{
+	return FInputCaptureUpdate::Ignore();
+}
+
+FInputCaptureUpdate UInputBehavior::UpdateCapture(const FInputDeviceState& input, const FInputCaptureData& data)
+{
+	return FInputCaptureUpdate::End();
+}
+
+void UInputBehavior::ForceEndCapture(const FInputCaptureData& data)
+{
+	return;
+}
+
+
+bool UInputBehavior::WantsHoverEvents()
+{
+	return false;
+}
+
+void UInputBehavior::UpdateHover(const FInputDeviceState& input)
+{
+	return;
+}
+
+void UInputBehavior::EndHover(const FInputDeviceState& input)
+{
+	return;
+}
+
diff --git a/Engine/Source/Runtime/InteractiveToolsFramework/Private/InputBehaviorSet.cpp b/Engine/Source/Runtime/InteractiveToolsFramework/Private/InputBehaviorSet.cpp
new file mode 100644
--- /dev/null
+++ b/Engine/Source/Runtime/InteractiveToolsFramework/Private/InputBehaviorSet.cpp
@@ -0,0 +1,164 @@
+// Copyright 1998-2019 Epic Games, Inc. All Rights Reserved.
+
+#include "InputBehaviorSet.h"
+
+
+UInputBehaviorSet::UInputBehaviorSet()
+{
+}
+
+UInputBehaviorSet::~UInputBehaviorSet()
+{
+}
+
+
+bool UInputBehaviorSet::IsEmpty() const
+{
+	return Behaviors.Num() == 0;
+}
+
+void UInputBehaviorSet::Add(UInputBehavior* Behavior, void* Source, const FString& Group)
+{
+	//if (source == null)
+	//	source = DefaultSource;
+
+	FBehaviorInfo Info;
+	Info.Behavior = Behavior;
+	Info.Source = Source;
+	Info.Group = Group;
+
+	Behaviors.Add(Info);
+	BehaviorsModified();
+}
+
+
+void UInputBehaviorSet::Add(const UInputBehaviorSet* OtherSet, void* NewSource, const FString& NewGroup)
+{
+	for (const FBehaviorInfo& OtherInfo : OtherSet->Behaviors)
+	{
+		FBehaviorInfo Info;
+		Info.Behavior = OtherInfo.Behavior;
+		Info.Source = (NewSource == nullptr) ? OtherInfo.Source : NewSource;
+		Info.Group = (NewGroup.IsEmpty()) ? OtherInfo.Group : NewGroup;
+		Behaviors.Add(Info);
+	}
+	if (OtherSet->Behaviors.Num() > 0)
+	{
+		BehaviorsModified();
+	}
+}
+
+
+bool UInputBehaviorSet::Remove(UInputBehavior* behavior)
+{
+	int32 idx = Behaviors.IndexOfByPredicate( 
+		[behavior] (const FBehaviorInfo& b) { return b.Behavior == behavior; }
+	);
+	if ( idx != INDEX_NONE ) {
+		Behaviors.RemoveAt(idx);
+		BehaviorsModified();
+		return true;
+	}
+	return false;
+}
+
+
+bool UInputBehaviorSet::RemoveByGroup(const FString& Group)
+{
+	int RemovedCount = Behaviors.RemoveAll(
+		[Group](const FBehaviorInfo& b) { return b.Group == Group; }
+	);
+	if (RemovedCount > 0) 
+	{
+		BehaviorsModified();
+	}
+	return (RemovedCount > 0);
+}
+
+
+
+bool UInputBehaviorSet::RemoveBySource(void* Source)
+{
+	int RemovedCount = Behaviors.RemoveAll(
+		[Source](const FBehaviorInfo& b) { return b.Source == Source; }
+	);
+	if (RemovedCount > 0)
+	{
+		BehaviorsModified();
+	}
+	return (RemovedCount > 0);	
+}
+
+
+void UInputBehaviorSet::RemoveAll()
+{
+	Behaviors.Reset();
+	BehaviorsModified();
+}
+
+
+void UInputBehaviorSet::CollectWantsCapture(const FInputDeviceState& input, TArray<FInputCaptureRequest>& result)
+{
+	for (auto b : Behaviors)
+	{
+		// only call WantsCapture if the Behavior supports the current input device
+		if (SupportsInputType(b.Behavior, input)) 
+		{
+			FInputCaptureRequest request = b.Behavior->WantsCapture(input);
+			if ( request.Type != EInputCaptureRequestType::Ignore ) 
+			{
+				request.Owner = b.Source;
+				result.Add(request);
+			}
+		}
+	}
+}
+
+
+
+bool UInputBehaviorSet::UpdateHover(const FInputDeviceState& input)
+{
+	bool bAnyWantedHover = false;
+	for (auto b : Behaviors) 
+	{
+		if (SupportsInputType(b.Behavior, input)) 
+		{
+			if (b.Behavior->WantsHoverEvents()) 
+			{
+				b.Behavior->UpdateHover(input);
+				bAnyWantedHover = true;
+			}
+		}
+	}
+	return bAnyWantedHover;
+}
+
+bool UInputBehaviorSet::EndHover(const FInputDeviceState& input)
+{
+	bool bAnyWantedHover = false;
+	for (auto b : Behaviors)
+	{
+		if (SupportsInputType(b.Behavior, input)) 
+		{
+			if (b.Behavior->WantsHoverEvents()) 
+			{
+				b.Behavior->EndHover(input);
+				bAnyWantedHover = true;
+			}
+		}
+	}
+	return bAnyWantedHover;
+}
+
+
+
+
+
+void UInputBehaviorSet::BehaviorsModified()
+{
+	// sort by priority
+	Behaviors.StableSort();
+
+	//send some kind of event...
+	//FUtil.SafeSendAnyEvent(OnSetChanged, this);
+}
diff --git a/Engine/Source/Runtime/InteractiveToolsFramework/Private/InputRouter.cpp b/Engine/Source/Runtime/InteractiveToolsFramework/Private/InputRouter.cpp
new file mode 100644
--- /dev/null
+++ b/Engine/Source/Runtime/InteractiveToolsFramework/Private/InputRouter.cpp
@@ -0,0 +1,211 @@
+// Copyright 1998-2019 Epic Games, Inc. All Rights Reserved.
+
+
+#include "InputRouter.h"
+
+
+UInputRouter::UInputRouter()
+{
+	ActiveLeftCapture = nullptr;
+	ActiveLeftCaptureOwner = nullptr;
+	ActiveRightCapture = nullptr;
+	ActiveRightCaptureOwner = nullptr;
+
+	bAutoInvalidateOnHover = false;
+	bAutoInvalidateOnCapture = false;
+
+	ActiveInputBehaviors = NewObject<UInputBehaviorSet>(this, "InputBehaviors");
+}
+
+
+
+void UInputRouter::Initialize(IToolsContextTransactionsAPI* TransactionsAPIIn)
+{
+	this->TransactionsAPI = TransactionsAPIIn;
+}
+
+void UInputRouter::Shutdown()
+{
+	this->TransactionsAPI = nullptr;
+}
+
+
+
+
+void UInputRouter::RegisterSource(IInputBehaviorSource* Source)
+{
+	ActiveInputBehaviors->Add(Source->GetInputBehaviors(), Source);
+}
+
+
+void UInputRouter::DeregisterSource(IInputBehaviorSource* Source)
+{
+	ActiveInputBehaviors->RemoveBySource(Source);
+}
+
+
+
+
+
+
+
+void UInputRouter::PostInputEvent(const FInputDeviceState& Input)
+{
+	// currently mouse-specific...
+
+	if ( ActiveInputBehaviors->IsEmpty() )
+	{
+		return;
+	}
+
+	if (Input.IsFromDevice(EInputDevices::Mouse) == false)
+	{
+		TransactionsAPI->PostMessage(TEXT("UInteractiveToolManager::PostInputEvent - non-mouse devices not currently supported."), EToolMessageLevel::Internal);
+		return;
+	}
+
+	if (ActiveLeftCapture != nullptr)
+	{
+		HandleCapturedMouseInput(Input);
+	}
+	else
+	{
+		ActiveLeftCaptureData = FInputCaptureData();
+		CheckForMouseCaptures(Input);
+	}
+
+	// update hover if nobody is capturing
+	if (ActiveLeftCapture == nullptr && ActiveRightCapture == nullptr)
+	{
+		bool bProcessed = ActiveInputBehaviors->UpdateHover(Input);
+		if (bProcessed && bAutoInvalidateOnHover)
+		{
+			TransactionsAPI->PostInvalidation();
+		}
+	}
+
+}
+
+
+void UInputRouter::PostHoverInputEvent(const FInputDeviceState& Input)
+{
+	LastHoverInput = Input;
+	bool bProcessed = ActiveInputBehaviors->UpdateHover(Input);
+	if (bProcessed && bAutoInvalidateOnHover)
+	{
+		TransactionsAPI->PostInvalidation();
+	}
+}
+
+
+
+
+bool UInputRouter::HasActiveMouseCapture() const
+{
+	return (ActiveLeftCapture != nullptr);
+}
+
+
+void UInputRouter::CheckForMouseCaptures(const FInputDeviceState& Input)
+{
+	TArray<FInputCaptureRequest> CaptureRequests;
+	ActiveInputBehaviors->CollectWantsCapture(Input, CaptureRequests);
+	if (CaptureRequests.Num() == 0)
+	{
+		return;
+	}
+
+	CaptureRequests.StableSort();
+
+	bool bAccepted = false;
+	for (int i = 0; i < CaptureRequests.Num() && bAccepted == false; ++i)
+	{
+		FInputCaptureUpdate Result =
+			CaptureRequests[i].Source->BeginCapture(Input, EInputCaptureSide::Left);
+		if (Result.State == EInputCaptureState::Begin)
+		{
+			// end outstanding hovers
+			ActiveInputBehaviors->EndHover(Input);
+
+			ActiveLeftCapture = Result.Source;
+			ActiveLeftCaptureOwner = CaptureRequests[i].Owner;
+			ActiveLeftCaptureData = Result.Data;
+			bAccepted = true;
+		}
+
+	}
+}
+
+
+void UInputRouter::HandleCapturedMouseInput(const FInputDeviceState& Input)
+{
+	if (ActiveLeftCapture == nullptr)
+	{
+		return;
+	}
+
+	// have active capture - give it this event
+
+	FInputCaptureUpdate Result =
+		ActiveLeftCapture->UpdateCapture(Input, ActiveLeftCaptureData);
+
+	if (Result.State == EInputCaptureState::End)
+	{
+		ActiveLeftCapture = nullptr;
+		ActiveLeftCaptureOwner = nullptr;
+		ActiveLeftCaptureData = FInputCaptureData();
+	}
+	else if (Result.State != EInputCaptureState::Continue)
+	{
+		TransactionsAPI->PostMessage(TEXT("UInteractiveToolManager::HandleCapturedMouseInput - unexpected capture state!"), EToolMessageLevel::Internal);
+	}
+
+	if (bAutoInvalidateOnCapture)
+	{
+		TransactionsAPI->PostInvalidation();
+	}
+}
+
+
+
+void UInputRouter::ForceTerminateAll()
+{
+	if (ActiveLeftCapture != nullptr)
+	{
+		ActiveLeftCapture->ForceEndCapture(ActiveLeftCaptureData);
+		ActiveLeftCapture = nullptr;
+		ActiveLeftCaptureOwner = nullptr;
+		ActiveLeftCaptureData = FInputCaptureData();
+	}
+
+	if (ActiveRightCapture != nullptr)
+	{
+		ActiveRightCapture->ForceEndCapture(ActiveRightCaptureData);
+		ActiveRightCapture = nullptr;
+		ActiveRightCaptureOwner = nullptr;
+		ActiveRightCaptureData = FInputCaptureData();
+	}
+
+	ActiveInputBehaviors->EndHover(LastHoverInput);
+}
+
+
+void UInputRouter::ForceTerminateSource(IInputBehaviorSource* Source)
+{
+	if (ActiveLeftCapture != nullptr && ActiveLeftCaptureOwner == Source)
+	{
+		ActiveLeftCapture->ForceEndCapture(ActiveLeftCaptureData);
+		ActiveLeftCapture = nullptr;
+		ActiveLeftCaptureOwner = nullptr;
+		ActiveLeftCaptureData = FInputCaptureData();
+	}
+
+	if (ActiveRightCapture != nullptr && ActiveRightCaptureOwner == Source)
+	{
+		ActiveRightCapture->ForceEndCapture(ActiveRightCaptureData);
+		ActiveRightCapture = nullptr;
+		ActiveRightCaptureOwner = nullptr;
+		ActiveRightCaptureData = FInputCaptureData();
+	}
+}
+
diff --git a/Engine/Source/Runtime/InteractiveToolsFramework/Private/InputState.cpp b/Engine/Source/Runtime/InteractiveToolsFramework/Private/InputState.cpp
new file mode 100644
--- /dev/null
+++ b/Engine/Source/Runtime/InteractiveToolsFramework/Private/InputState.cpp
@@ -0,0 +1,4 @@
+// Copyright 1998-2019 Epic Games, Inc. All Rights Reserved.
+
+#include "InputState.h"
+
diff --git a/Engine/Source/Runtime/InteractiveToolsFramework/Private/InteractiveTool.cpp b/Engine/Source/Runtime/InteractiveToolsFramework/Private/InteractiveTool.cpp
new file mode 100644
--- /dev/null
+++ b/Engine/Source/Runtime/InteractiveToolsFramework/Private/InteractiveTool.cpp
@@ -0,0 +1,82 @@
+// Copyright 1998-2019 Epic Games, Inc. All Rights Reserved.
+
+
+#include "InteractiveTool.h"
+#include "InteractiveToolManager.h"
+
+
+UInteractiveTool::UInteractiveTool()
+{
+	// tools need to be transactional or undo/redo won't work on their uproperties
+	SetFlags(RF_Transactional);
+
+	// tools don't get saved but this isn't necessary because they are created in the transient package...
+	//SetFlags(RF_Transient);
+
+	InputBehaviors = NewObject<UInputBehaviorSet>(this, TEXT("InputBehaviors"));
+}
+
+void UInteractiveTool::Setup()
+{
+}
+
+void UInteractiveTool::Shutdown(EToolShutdownType ShutdownType)
+{
+	InputBehaviors->RemoveAll();
+	ToolPropertyObjects.Reset();
+}
+
+void UInteractiveTool::Render(IToolsContextRenderAPI* RenderAPI)
+{
+}
+
+
+void UInteractiveTool::AddInputBehavior(UInputBehavior* Behavior)
+{
+	InputBehaviors->Add(Behavior);
+}
+
+const UInputBehaviorSet* UInteractiveTool::GetInputBehaviors() const
+{
+	return InputBehaviors;
+}
+
+
+void UInteractiveTool::AddToolPropertySource(UObject* PropertyObject)
+{
+	ToolPropertyObjects.Add(PropertyObject);
+}
+
+const TArray<UObject*>& UInteractiveTool::GetToolProperties() const
+{
+	return ToolPropertyObjects;
+}
+
+
+bool UInteractiveTool::HasCancel() const
+{
+	return false;
+}
+
+bool UInteractiveTool::HasAccept() const
+{
+	return false;
+}
+
+bool UInteractiveTool::CanAccept() const
+{
+	return false;
+}
+
+
+void UInteractiveTool::Tick(float DeltaTime)
+{
+}
+
+
+UInteractiveToolManager* UInteractiveTool::GetToolManager() const
+{
+	UInteractiveToolManager* ToolManager = Cast<UInteractiveToolManager>(GetOuter());
+	check(ToolManager != nullptr);
+	return ToolManager;
+}
\ No newline at end of file
diff --git a/Engine/Source/Runtime/InteractiveToolsFramework/Private/InteractiveToolBuilder.cpp b/Engine/Source/Runtime/InteractiveToolsFramework/Private/InteractiveToolBuilder.cpp
new file mode 100644
--- /dev/null
+++ b/Engine/Source/Runtime/InteractiveToolsFramework/Private/InteractiveToolBuilder.cpp
@@ -0,0 +1,6 @@
+// Copyright 1998-2019 Epic Games, Inc. All Rights Reserved.
+
+
+#include "InteractiveToolBuilder.h"
+
+// Add default functionality here for any UInteractiveToolBuilder functions that are not pure virtual.
diff --git a/Engine/Source/Runtime/InteractiveToolsFramework/Private/InteractiveToolManager.cpp b/Engine/Source/Runtime/InteractiveToolsFramework/Private/InteractiveToolManager.cpp
new file mode 100644
--- /dev/null
+++ b/Engine/Source/Runtime/InteractiveToolsFramework/Private/InteractiveToolManager.cpp
@@ -0,0 +1,243 @@
+// Copyright 1998-2019 Epic Games, Inc. All Rights Reserved.
+
+
+#include "InteractiveToolManager.h"
+
+
+UInteractiveToolManager::UInteractiveToolManager()
+{
+	QueriesAPI = nullptr;
+	TransactionsAPI = nullptr;
+	InputRouter = nullptr;
+
+	ActiveLeftBuilder = nullptr;
+	ActiveLeftTool = nullptr;
+
+	ActiveRightBuilder = nullptr;
+	ActiveRightTool = nullptr;
+}
+
+
+void UInteractiveToolManager::Initialize(IToolsContextQueriesAPI* queriesAPI, IToolsContextTransactionsAPI* transactionsAPI, UInputRouter* InputRouterIn)
+{
+	this->QueriesAPI = queriesAPI;
+	this->TransactionsAPI = transactionsAPI;
+	this->InputRouter = InputRouterIn;
+}
+
+
+void UInteractiveToolManager::Shutdown()
+{
+	this->QueriesAPI = nullptr;
+
+	if (ActiveLeftTool != nullptr)
+	{
+		DeactivateTool(EToolSide::Left, EToolShutdownType::Cancel);
+	}
+	if (ActiveRightTool != nullptr)
+	{
+		DeactivateTool(EToolSide::Right, EToolShutdownType::Cancel);
+	}
+
+	this->TransactionsAPI = nullptr;
+}
+
+
+
+void UInteractiveToolManager::RegisterToolType(const FString& Identifier, UInteractiveToolBuilder* Builder)
+{
+	check(ToolBuilders.Contains(Identifier) == false);
+	ToolBuilders.Add(Identifier, Builder );
+}
+
+
+bool UInteractiveToolManager::SelectActiveToolType(EToolSide Side, const FString& Identifier)
+{
+	if (ToolBuilders.Contains(Identifier))
+	{
+		UInteractiveToolBuilder* Builder = ToolBuilders[Identifier];
+		if (Side == EToolSide::Right)
+		{
+			ActiveRightBuilder = Builder;
+		}
+		else
+		{
+			ActiveLeftBuilder = Builder;
+		}
+		return true;
+	}
+	return false;
+}
+
+
+
+bool UInteractiveToolManager::CanActivateTool(EToolSide Side, const FString& Identifier)
+{
+	check(Side == EToolSide::Left);   // TODO: support right-side tool
+
+	if (ActiveLeftTool != nullptr)
+	{
+		return false;
+	}
+
+	if (ToolBuilders.Contains(Identifier))
+	{
+		FToolBuilderState InputState;
+		QueriesAPI->GetCurrentSelectionState(InputState);
+
+		UInteractiveToolBuilder* Builder = ToolBuilders[Identifier];
+		return Builder->CanBuildTool(InputState);
+	}
+
+	return false;
+}
+
+
+bool UInteractiveToolManager::ActivateTool(EToolSide Side)
+{
+	check(Side == EToolSide::Left);   // TODO: support right-side tool
+
+	if (ActiveLeftTool != nullptr) 
+	{
+		DeactivateTool(EToolSide::Left, EToolShutdownType::Cancel);
+	}
+
+	if (ActiveLeftBuilder == nullptr) 
+	{
+		return false;
+	}
+
+	// construct input state we will pass to tools
+	FToolBuilderState InputState;
+	QueriesAPI->GetCurrentSelectionState(InputState);
+
+	if (ActiveLeftBuilder->CanBuildTool(InputState) == false)
+	{
+		TransactionsAPI->PostMessage( TEXT("UInteractiveToolManager::ActivateTool: CanBuildTool returned false."), EToolMessageLevel::Internal);
+		return false;
+	}
+
+	ActiveLeftTool = ActiveLeftBuilder->BuildTool(InputState);
+	if (ActiveLeftTool == nullptr)
+	{
+		return false;
+	}
+
+	ActiveLeftTool->Setup();
+
+	// register new active input behaviors
+	InputRouter->RegisterSource(ActiveLeftTool);
+
+	PostInvalidation();
+
+	return true;
+}
+
+
+void UInteractiveToolManager::DeactivateTool(EToolSide Side, EToolShutdownType ShutdownType)
+{
+	check(Side == EToolSide::Left);   // TODO: support right-side tool
+
+	if (ActiveLeftTool != nullptr)
+	{
+		InputRouter->ForceTerminateSource(ActiveLeftTool);
+
+		ActiveLeftTool->Shutdown(ShutdownType);
+
+		InputRouter->DeregisterSource(ActiveLeftTool);
+
+		ActiveLeftTool = nullptr;
+
+		PostInvalidation();
+	}
+}
+
+
+bool UInteractiveToolManager::HasActiveTool(EToolSide Side)
+{
+	return (Side == EToolSide::Left) ? (ActiveLeftTool != nullptr) : (ActiveRightTool != nullptr);
+}
+
+UInteractiveTool* UInteractiveToolManager::GetActiveTool(EToolSide Side)
+{
+	return (Side == EToolSide::Left) ? ActiveLeftTool : ActiveRightTool;
+}
+
+
+bool UInteractiveToolManager::CanAcceptActiveTool(EToolSide Side)
+{
+	if (ActiveLeftTool != nullptr)
+	{
+		return ActiveLeftTool->HasAccept() && ActiveLeftTool->CanAccept();
+	}
+	return false;
+}
+
+bool UInteractiveToolManager::CanCancelActiveTool(EToolSide Side)
+{
+	if (ActiveLeftTool != nullptr)
+	{
+		return ActiveLeftTool->HasCancel();
+	}
+	return false;
+}
+
+
+
+
+
+
+void UInteractiveToolManager::Tick(float DeltaTime)
+{
+	if (ActiveLeftTool != nullptr)
+	{
+		ActiveLeftTool->Tick(DeltaTime);
+	}
+
+	if (ActiveRightTool!= nullptr)
+	{
+		ActiveRightTool->Tick(DeltaTime);
+	}
+}
+
+
+void UInteractiveToolManager::Render(IToolsContextRenderAPI* RenderAPI)
+{
+	if (ActiveLeftTool != nullptr)
+	{
+		ActiveLeftTool->Render(RenderAPI);
+	}
+
+	if (ActiveRightTool != nullptr)
+	{
+		ActiveRightTool->Render(RenderAPI);
+	}
+}
+
+
+
+
+void UInteractiveToolManager::PostInvalidation()
+{
+	TransactionsAPI->PostInvalidation();
+}
+
+
+void UInteractiveToolManager::BeginUndoTransaction(const FText& Description)
+{
+	TransactionsAPI->BeginUndoTransaction(Description);
+}
+
+void UInteractiveToolManager::EndUndoTransaction()
+{
+	TransactionsAPI->EndUndoTransaction();
+}
+
+
+
+void UInteractiveToolManager::EmitObjectChange(UObject* TargetObject, TUniquePtr<FChange> Change, const FText& Description)
+{
+	TransactionsAPI->AppendChange(TargetObject, MoveTemp(Change), Description );
+}
+
+
diff --git a/Engine/Source/Runtime/InteractiveToolsFramework/Private/InteractiveToolsContext.cpp b/Engine/Source/Runtime/InteractiveToolsFramework/Private/InteractiveToolsContext.cpp
new file mode 100644
--- /dev/null
+++ b/Engine/Source/Runtime/InteractiveToolsFramework/Private/InteractiveToolsContext.cpp
@@ -0,0 +1,32 @@
+// Copyright 1998-2019 Epic Games, Inc. All Rights Reserved.
+
+
+#include "InteractiveToolsContext.h"
+
+
+UInteractiveToolsContext::UInteractiveToolsContext()
+{
+	InputRouter = nullptr;
+	ToolManager = nullptr;
+}
+
+void UInteractiveToolsContext::Initialize(IToolsContextQueriesAPI* QueriesAPI, IToolsContextTransactionsAPI* TransactionsAPI)
+{
+	InputRouter = NewObject<UInputRouter>(this);
+	InputRouter->Initialize(TransactionsAPI);
+
+	ToolManager = NewObject<UInteractiveToolManager>(this);
+	ToolManager->Initialize(QueriesAPI, TransactionsAPI, InputRouter);
+}
+
+
+void UInteractiveToolsContext::Shutdown()
+{
+	// force-terminate any remaining captures/hovers/etc
+	InputRouter->ForceTerminateAll();
+	InputRouter->Shutdown();
+	InputRouter = nullptr;
+
+	ToolManager->Shutdown();
+	ToolManager = nullptr;
+}
diff --git a/Engine/Source/Runtime/InteractiveToolsFramework/Private/InteractiveToolsFramework.cpp b/Engine/Source/Runtime/InteractiveToolsFramework/Private/InteractiveToolsFramework.cpp
new file mode 100644
--- /dev/null
+++ b/Engine/Source/Runtime/InteractiveToolsFramework/Private/InteractiveToolsFramework.cpp
@@ -0,0 +1,20 @@
+// Copyright 1998-2018 Epic Games, Inc. All Rights Reserved.
+
+#include "InteractiveToolsFramework.h"
+
+#define LOCTEXT_NAMESPACE "FInteractiveToolsFrameworkModule"
+
+void FInteractiveToolsFrameworkModule::StartupModule()
+{
+	// This code will execute after your module is loaded into memory; the exact timing is specified in the .uplugin file per-module
+}
+
+void FInteractiveToolsFrameworkModule::ShutdownModule()
+{
+	// This function may be called during shutdown to clean up your module.  For modules that support dynamic reloading,
+	// we call this function before unloading the module.
+}
+
+#undef LOCTEXT_NAMESPACE
+	
+IMPLEMENT_MODULE(FInteractiveToolsFrameworkModule, InteractiveToolsFramework)
\ No newline at end of file
diff --git a/Engine/Source/Runtime/InteractiveToolsFramework/Private/ToolBuilderUtil.cpp b/Engine/Source/Runtime/InteractiveToolsFramework/Private/ToolBuilderUtil.cpp
new file mode 100644
--- /dev/null
+++ b/Engine/Source/Runtime/InteractiveToolsFramework/Private/ToolBuilderUtil.cpp
@@ -0,0 +1,93 @@
+// Copyright 1998-2019 Epic Games, Inc. All Rights Reserved.
+
+#include "ToolBuilderUtil.h"
+#include "CoreMinimal.h"
+#include "GameFramework/Actor.h"
+#include "Engine/Selection.h"
+#include "Components/StaticMeshComponent.h"
+
+
+
+bool ToolBuilderUtil::IsMeshDescriptionSourceComponent(UObject* ComponentObject)
+{
+	UStaticMeshComponent* StaticMeshComp = Cast<UStaticMeshComponent>(ComponentObject);
+	if (StaticMeshComp != nullptr)
+	{
+		return true;
+	}
+	return false;
+}
+
+
+
+int ToolBuilderUtil::CountComponents(const FToolBuilderState& InputState, const TFunction<bool(UObject*)>& Predicate)
+{
+	int nTypedComponents = 0;
+
+	if (InputState.SelectedComponents != nullptr && InputState.SelectedComponents->Num() > 0)
+	{
+		for (FSelectionIterator Iter(*InputState.SelectedComponents); Iter; ++Iter)
+		{
+			if (Predicate(*Iter))
+			{
+				nTypedComponents++;
+			}
+		}
+	}
+	else
+	{
+		for (FSelectionIterator Iter(*InputState.SelectedActors); Iter; ++Iter)
+		{
+			if (AActor* Actor = Cast<AActor>(*Iter))
+			{
+				for (UActorComponent* Component : Actor->GetComponents())
+				{
+					if (Predicate(Component))
+					{
+						nTypedComponents++;
+					}
+				}
+			}
+		}
+	}
+
+	return nTypedComponents;
+}
+
+
+
+
+UObject* ToolBuilderUtil::FindFirstComponent(const FToolBuilderState& InputState, const TFunction<bool(UObject*)>& Predicate)
+{
+	if (InputState.SelectedComponents != nullptr && InputState.SelectedComponents->Num() > 0)
+	{
+		for (FSelectionIterator Iter(*InputState.SelectedComponents); Iter; ++Iter)
+		{
+			if (Predicate(*Iter))
+			{
+				return *Iter;
+			}
+		}
+	}
+	else
+	{
+		for (FSelectionIterator Iter(*InputState.SelectedActors); Iter; ++Iter)
+		{
+			if (AActor* Actor = Cast<AActor>(*Iter))
+			{
+				for (UActorComponent* Component : Actor->GetComponents())
+				{
+					if (Predicate(Component))
+					{
+						return Component;
+					}
+				}
+			}
+		}
+	}
+
+	return nullptr;
+}
+
+
+
diff --git a/Engine/Source/Runtime/InteractiveToolsFramework/Public/ComponentSourceInterfaces.h b/Engine/Source/Runtime/InteractiveToolsFramework/Public/ComponentSourceInterfaces.h
new file mode 100644
--- /dev/null
+++ b/Engine/Source/Runtime/InteractiveToolsFramework/Public/ComponentSourceInterfaces.h
@@ -0,0 +1,99 @@
+// Copyright 1998-2019 Epic Games, Inc. All Rights Reserved.
+
+#pragma once
+
+#include "CoreMinimal.h"
+#include "Math/UnrealMath.h"
+#include "Engine/EngineTypes.h"    // FHitResult
+
+// predeclarations
+class AActor;
+struct FMeshDescription;
+class UMaterialInterface;
+
+
+
+/**
+ * Wrapper around a UObject Component that can provide a MeshDescription, and
+ * (optionally) bake a modified MeshDescription back to this Component.
+ * An example of a Source might be a StaticMeshComponent. How a modified
+ * MeshDescription is committed back is context-dependent (in Editor vs PIE vs Runtime, etc).
+ * 
+ * (Conceivably this doesn't have to be backed by a Component, but most usage will assume there is an Actor)
+ */
+class IMeshDescriptionSource
+{
+public:
+	virtual ~IMeshDescriptionSource() {}
+
+	
+	/** @return the Actor that owns this Component */
+	virtual AActor* GetOwnerActor() const = 0;
+
+	/** @return the Component this is a Source for */
+	virtual UObject* GetOwnerComponent() const = 0;
+
+	/** @return Pointer to the MeshDescription this Source is providing */
+	virtual FMeshDescription* GetMeshDescription() const = 0;
+
+	/**
+	 * Get pointer to a Material provided by this Source
+	 * @param MaterialIndex index of the material
+	 * @return MaterialInterface pointer, or null if MaterialIndex is invalid
+	 */
+	virtual UMaterialInterface* GetMaterial(int32 MaterialIndex) const = 0;
+
+	/** 
+	 * @return the transform on this component
+	 * @todo Do we need to return a list of transforms here?
+	 */
+	virtual FTransform GetWorldTransform() const = 0;
+
+	/**
+	 * Compute ray intersection with the MeshDescription this Source is providing
+	 * @param WorldRay ray in world space
+	 * @param OutHit hit test data
+	 * @return true if ray intersected Component
+	 */
+	virtual bool HitTest(const FRay& WorldRay, FHitResult& OutHit) const = 0;
+
+	/**
+	 * Set the visibility of the Component associated with this Source (ie to hide during Tool usage)
+	 * @param bVisible desired visibility
+	 */
+	virtual void SetOwnerVisibility(bool bVisible) const = 0;
+
+
+	/** @return true if this Source is read-only, ie Commit functions cannot be called */
+	virtual bool IsReadOnly() const { return true; }
+
+	/**
+	 * Call this to modify the MeshDescription provided by the Source. You provide a callback
+	 * function that the Source will immediately call with a suitable MeshDescription instance.
+	 * The Source will then update the Component as necessary. The Source is responsible for
+	 * making sure that it is safe to modify this MeshDescription instance.
+	 * @param ModifyFunction callback function that updates/modifies a MeshDescription
+	 */
+	virtual void CommitInPlaceModification(const TFunction<void(FMeshDescription*)>& ModifyFunction) { check(false); }
+};
+
+
+
+
+
+/**
+ * Interface to a Factory that knows how to build Sources for UObjects. 
+ */
+class IComponentSourceFactory
+{
+public:
+	virtual ~IComponentSourceFactory() {}
+
+	/**
+	 * Create a MeshDescription source for the given Component
+	 * @param Component A UObject that can provide a MeshDescription. Assumption is this is a Component of an Actor.
+	 * @return A MeshDescriptionSource instance. Must not return null.
+	 */
+	virtual TUniquePtr<IMeshDescriptionSource> MakeMeshDescriptionSource(UObject* Component) = 0;
+};
+
diff --git a/Engine/Source/Runtime/InteractiveToolsFramework/Public/InputBehavior.h b/Engine/Source/Runtime/InteractiveToolsFramework/Public/InputBehavior.h
new file mode 100644
--- /dev/null
+++ b/Engine/Source/Runtime/InteractiveToolsFramework/Public/InputBehavior.h
@@ -0,0 +1,282 @@
+// Copyright 1998-2019 Epic Games, Inc. All Rights Reserved.
+
+#pragma once
+
+#include "CoreMinimal.h"
+#include "InputState.h"
+#include "Math/NumericLimits.h"
+#include "InputBehavior.generated.h"
+
+
+
+/**
+ * Input can be captured separately for Left and Right sides (eg for VR controllers)
+ * Currently mouse is Left.
+ */
+UENUM()
+enum class EInputCaptureSide
+{
+	None = 0,
+	Left = 1,
+	Right = 2,
+	Both = 3,
+	Any = 99
+};
+
+
+/**
+ * An active capturing behavior may need to keep track of additional data that
+ * cannot be stored within the behavior (for example if the same behavior instance
+ * is capturing for Left and Right separately). So FInputCaptureUpdate can optionally
+ * return this structure, and we will pass it to the next UpdateCapture() call
+ */
+struct INTERACTIVETOOLSFRAMEWORK_API FInputCaptureData
+{
+	/** Which side do we want to capture on */
+	EInputCaptureSide WhichSide;
+	/** pointer to data defined by the InputBehavior, which is also responsible for cleaning it up */
+	void* CustomData;
+
+	FInputCaptureData()
+	{
+		WhichSide = EInputCaptureSide::None;
+		CustomData = nullptr;
+	}
+};
+
+
+/**
+ * Used by FInputCaptureRequest to indicate whether the InputBehavior
+ * wants to capture or ignore an input event
+ */
+UENUM()
+enum class EInputCaptureRequestType
+{
+	Begin = 1,
+	Ignore = 2
+};
+
+
+// predeclaration
+class UInputBehavior;
+
+
+/**
+ * UInputBehavior returns an FInputCaptureRequest from WantsCapture() to indicate
+ * whether it wants to capture or ignore an input event
+ */
+struct INTERACTIVETOOLSFRAMEWORK_API FInputCaptureRequest
+{
+	/** Which input behavior generated this request */
+	UInputBehavior* Source;
+    /** What type of capture request is this (Begin or Ignore) */
+	EInputCaptureRequestType Type;
+	/** Which side does request want to capture on */
+	EInputCaptureSide Side;
+
+	/** Depth along hit-test ray */
+	float HitDepth;
+
+	/** Owner of the requesting behavior. Behavior doesn't know this, so this is initialized to null */
+	void* Owner;
+
+
+	FInputCaptureRequest(EInputCaptureRequestType type, UInputBehavior* behavior, EInputCaptureSide whichSide, float hitDepth = TNumericLimits<float>::Max() )
+	{
+		this->Type = type;
+		this->Source = behavior;
+		this->Side = whichSide;
+		this->Owner = nullptr;
+		this->HitDepth = hitDepth;
+	}
+
+	/** Create a Begin-capture request */
+	static FInputCaptureRequest Begin(UInputBehavior* behavior, EInputCaptureSide whichSide, float hitDepth = TNumericLimits<float>::Max() )
+	{
+		return FInputCaptureRequest(EInputCaptureRequestType::Begin, behavior, whichSide, hitDepth);
+	}
+
+	/** Create an ignore-capture request */
+	static FInputCaptureRequest Ignore() 
+	{
+		return FInputCaptureRequest(EInputCaptureRequestType::Ignore, nullptr, EInputCaptureSide::Any, TNumericLimits<float>::Max() );
+	}
+
+	friend bool operator<(const FInputCaptureRequest& l, const FInputCaptureRequest& r);
+};
+
+
+
+
+
+
+/**
+ * FInputCaptureUpdate uses this type to indicate what state the capturing Behavior
+ * would like to transition to, based on the input event
+ */
+UENUM()
+enum class EInputCaptureState
+{
+	Begin = 1,		// start capturing (which should always be the case if BeginCapture is called)
+	Continue = 2,   // Behavior wants to continue capturing
+	End = 3,        // Behavior wants to end capturing
+	Ignore = 4      // Behavior ignored this event
+};
+
+
+
+/**
+ * IInputBehavior returns an FInputCaptureUpdate from BeginCapture() and UpdateCapture(),
+ * which indicates to the InputRouter what the Behavior would like to have happen.
+ */
+struct INTERACTIVETOOLSFRAMEWORK_API FInputCaptureUpdate
+{
+	/** Indicates what capture state the Behavior wants to transition to */
+	EInputCaptureState State;
+	/** Which Behavior did this update come from */
+	UInputBehavior* Source;
+	/** custom data for the active capture that should be propagated to next UpdateCapture() call */
+	FInputCaptureData Data;
+
+	/** 
+	 * Create a begin-capturing instance of FInputCaptureUpdate
+	 * @param Source UInputBehavior that is returning this update
+	 * @param Which Which side we are capturing on
+	 * @param CustomData client-provided data that will be passed to UInputBehavior::UpdateCapture() calls. Client owns this memory!
+	 */
+	static FInputCaptureUpdate Begin(UInputBehavior* SourceBehavior, EInputCaptureSide WhichSide, void* CustomData = nullptr)
+	{
+		return FInputCaptureUpdate(EInputCaptureState::Begin, SourceBehavior, WhichSide, CustomData);
+	}
+
+	/** Create a default continue-capturing instance of FInputCaptureUpdate */
+	static FInputCaptureUpdate Continue()
+	{
+		return FInputCaptureUpdate(EInputCaptureState::Continue, nullptr, EInputCaptureSide::Any);
+	}
+	/** Create a default end-capturing instance of FInputCaptureUpdate */
+	static FInputCaptureUpdate End()
+	{
+		return FInputCaptureUpdate(EInputCaptureState::End, nullptr, EInputCaptureSide::Any);
+	}
+	/** Create a default ignore-capturing instance of FInputCaptureUpdate */
+	static FInputCaptureUpdate Ignore()
+	{
+		return FInputCaptureUpdate(EInputCaptureState::Ignore, nullptr, EInputCaptureSide::Any);
+	}
+
+	/**
+	 * @param StateIn desired capture state
+	 * @param Source UInputBehavior that is returning this update
+	 * @param Which Which side we are capturing on
+	 * @param CustomData client-provided data that will be passed to UInputBehavior::UpdateCapture() calls. Client owns this memory!
+	 */
+	FInputCaptureUpdate(EInputCaptureState StateIn, UInputBehavior* SourceBehaviorIn, EInputCaptureSide WhichSideIn, void* CustomData = nullptr)
+	{
+		State = StateIn;
+		Source = SourceBehaviorIn;
+		Data.WhichSide = WhichSideIn;
+		Data.CustomData = CustomData;
+	}
+};
+
+
+
+/**
+ * Each UInputBehavior provides a priority that is used to help resolve situations
+ * when multiple Behaviors want to capture based on the same input event
+ */
+struct INTERACTIVETOOLSFRAMEWORK_API FInputCapturePriority
+{
+	/** Constant priority value */
+	int Priority;
+
+	FInputCapturePriority(int priority = 100) 
+	{
+		Priority = priority;
+	}
+
+	friend bool operator<(const FInputCapturePriority& l, const FInputCapturePriority& r)
+	{
+		return l.Priority < r.Priority;
+	}
+	friend bool operator==(const FInputCapturePriority& l, const FInputCapturePriority& r)
+	{
+		return l.Priority == r.Priority;
+	}
+};
+
+
+
+/**
+ * An InputBehavior implements a state machine for a user interaction. 
+ * The InputRouter maintains a set of active Behaviors, and when new input
+ * events occur, it calls WantsCapture() to check if the Behavior would like to
+ * begin capturing the applicable input event stream (eg for a mouse, one or both VR controllers, etc).
+ * If the Behavior acquires capture, UpdateCapture() is called until the Behavior
+ * indicates that it wants to release the device, or until the InputRouter force-terminates
+ * the capture via ForceEndCapture().
+ *
+ * For example, something like ButtonSetClickBehavior might work as follows:
+ *    - in WantsCapture(), if left mouse is pressed and a button is under cursor, return Begin, otherwise Ignore
+ *    - in BeginCapture(), save identifier for button that is under cursor
+ *    - in UpdateCapture()
+ *        - if left mouse is down, return Continue
+ *        - if left mouse is released:
+ *            - if saved button is still under cursor, call button.Clicked()
+ *            - return End
+ *
+ * Written sufficiently generically, the above Behavior doesn't need to know about buttons,
+ * it just needs to know how to hit-test the clickable object(s). Similarly separate 
+ * Behaviors can be written for mouse, VR, touch, gamepad, etc. 
+ *
+ * Implementing interactions in this way allows the input handling to be separated from functionality.
+ */
+UCLASS(Transient)
+class INTERACTIVETOOLSFRAMEWORK_API UInputBehavior : public UObject
+{
+	GENERATED_BODY()
+
+public:
+	UInputBehavior();
+
+	/** The priority is used to resolve situations where multiple behaviors want the same capture */
+	virtual FInputCapturePriority GetPriority();
+
+	/** Which device types does this Behavior support */
+	virtual EInputDevices GetSupportedDevices();
+
+	/** Given the input state, does this Behavior want to begin capturing some input devices? */
+	virtual FInputCaptureRequest WantsCapture(const FInputDeviceState& InputState);
+
+	/** Called after WantsCapture() returns a capture request that was accepted */
+	virtual FInputCaptureUpdate BeginCapture(const FInputDeviceState& InputState, EInputCaptureSide eSide);
+
+	/** 
+	 * Called for each new input event during a capture sequence. Return Continue to keep
+	 * capturing, or End to finish capturing.
+	 */
+	virtual FInputCaptureUpdate UpdateCapture(const FInputDeviceState& InputState, const FInputCaptureData& CaptureData);
+
+	/** If this is called, the Behavior has forcibly lost capture (eg due to app losing focus for example) and needs to clean up accordingly */
+	virtual void ForceEndCapture(const FInputCaptureData& CaptureData);
+
+
+	//
+	// hover support (optional)
+	//
+
+	/** return true if this Behavior supports hover (ie passive input events) */
+	virtual bool WantsHoverEvents();
+
+	/** called on each new hover input event, ie if no other behavior is actively capturing input */
+	virtual void UpdateHover(const FInputDeviceState& InputState);
+
+	/** if a capture begins or focus is lost, any active hover visualization needs to terminate */
+	virtual void EndHover(const FInputDeviceState& InputState);
+
+
+protected:
+	/** priority returned by GetPriority() */
+	FInputCapturePriority DefaultPriority;
+};
\ No newline at end of file
diff --git a/Engine/Source/Runtime/InteractiveToolsFramework/Public/InputBehaviorSet.h b/Engine/Source/Runtime/InteractiveToolsFramework/Public/InputBehaviorSet.h
new file mode 100644
--- /dev/null
+++ b/Engine/Source/Runtime/InteractiveToolsFramework/Public/InputBehaviorSet.h
@@ -0,0 +1,179 @@
+// Copyright 1998-2019 Epic Games, Inc. All Rights Reserved.
+
+#pragma once
+
+#include "CoreMinimal.h"
+#include "UObject/Interface.h"
+#include "InputBehavior.h"
+#include "InputBehaviorSet.generated.h"
+
+
+
+/**
+ * This is an internal structure used by UInputBehaviorSet. 
+ */
+USTRUCT()
+struct INTERACTIVETOOLSFRAMEWORK_API FBehaviorInfo
+{
+	GENERATED_BODY()
+	
+	/** Reference to a Behavior */
+	UPROPERTY()
+	UInputBehavior* Behavior;
+
+	/** Source object that provided this Behavior */
+	void* Source;
+
+	/** Group identifier for this Behavior */
+	FString Group;
+
+	friend bool operator<(const FBehaviorInfo& l, const FBehaviorInfo& r)
+	{
+		return l.Behavior->GetPriority() < r.Behavior->GetPriority();
+	}
+};
+
+
+
+
+/**
+ * UInputBehaviorSet manages a set of UInputBehaviors, and provides various functions
+ * to query and forward events to the set. Tools and Widgets provide instances of this via
+ * IInputBehaviorSource, and UInputRouter collects and manages them (see comments there)
+ *
+ * Behaviors in the set each have a source pointer and group tag, which allows sets of
+ * behaviors to be managed together. For example one UInputBehaviorSet can be merged into 
+ * another and removed later.
+ */
+UCLASS(Transient)
+class INTERACTIVETOOLSFRAMEWORK_API UInputBehaviorSet : public UObject
+{
+	GENERATED_BODY()
+
+public:
+	UInputBehaviorSet();
+	virtual ~UInputBehaviorSet();
+
+	//
+	// Set Management
+	//
+
+
+	/** @return true if there are no Behaviors in set */
+	virtual bool IsEmpty() const;
+
+	/** 
+	 * Add a Behavior to the set
+	 * @param Behavior Behavior to add to set
+	 * @param Source pointer to owning object, used only to identify Behavior later
+	 * @param GroupName string identifier for this Behavior or group of Behaviors
+	 */
+	virtual void Add(UInputBehavior* Behavior, void* Source = nullptr, const FString& GroupName = "");
+
+	/**
+	 * Merge another BehaviorSet into this Set
+	 * @param OtherSet Set of Behaviors to add to this Set
+	 * @param Source pointer to owning object, used only to identify Behavior later. If nullptr, source is copied from other Set.
+	 * @param NewGroupName string identifier for this Behavior or group of Behaviors. If empty string, group is copied from other Set.
+	 */
+	virtual void Add(const UInputBehaviorSet* OtherSet, void* NewSource = nullptr, const FString& NewGroupName = "");
+
+	/**
+	 * Remove a Behavior from the Set
+	 * @param Behavior Behavior to remove
+	 * @return true if Behavior was found and removed
+	 */
+	virtual bool Remove(UInputBehavior* Behavior);
+	/**
+	 * Remove a group of Behaviors from the Set
+	 * @param GroupName name of group, all Behaviors that were added with this GroupName are removed
+	 * @return true if any Behaviors were found and removed
+	 */
+	virtual bool RemoveByGroup(const FString& GroupName);
+	/**
+	 * Remove a group of Behaviors from the Set
+	 * @param Source source object pointer, all Behaviors that were added with this Source pointer are removed.
+	 * @return true if any Behaviors were found and removed
+	 */
+	virtual bool RemoveBySource(void* Source);
+	/**
+	 * Remove all Behaviors from the set
+	 */
+	virtual void RemoveAll();
+
+
+	//
+	// Queries and Event Forwarding
+	//
+
+
+	/**
+	 * Call UInputBehavior::WantsCapture() on each valid Behavior and collect up the requests that indicated a Capture was desired.
+	 * @param InputState current input device state information
+	 * @param ResultOut returned set of non-Ignoring FInputCaptureRequests returned by Behaviors
+	 */
+	virtual void CollectWantsCapture(const FInputDeviceState& InputState, TArray<FInputCaptureRequest>& ResultOut);
+
+	/**
+	 * Call UInputBehavior::UpdateHover() on each valid Behavior that wants hover events
+	 * @param InputState current input device state information
+	 * @return true if any InputBehavior indicated that it wants hover events
+	 */
+	virtual bool UpdateHover(const FInputDeviceState&  InputState);
+
+	/**
+	 * Call UInputBehavior::EndHover() on each valid Behavior that wants hover events
+	 * @param InputState current input device state information
+	 * @return true if any InputBehavior indicated that it wants hover events
+	 */
+	virtual bool EndHover(const FInputDeviceState&  InputState);
+
+
+protected:
+
+	/** Current set of known Behaviors */
+	UPROPERTY()
+	TArray<FBehaviorInfo> Behaviors;
+
+	/**
+	 * called internally when Behaviors list is updated, to re-sort by Priority, etc (@todo: and emit change event)
+	 */
+	virtual void BehaviorsModified();
+
+
+	/** @return true if Behavior supports InputState.InputDevice */
+	bool SupportsInputType(UInputBehavior* Behavior, const FInputDeviceState& InputState)
+	{
+		return (Behavior->GetSupportedDevices() & InputState.InputDevice) != EInputDevices::None;
+	}
+
+};
+
+
+
+
+
+// UInterface for IInputBehavior
+UINTERFACE(MinimalAPI)
+class UInputBehaviorSource : public UInterface
+{
+	GENERATED_BODY()
+};
+
+
+/** 
+ * UObjects that implement IInputBehaviorSource have an UInputBehaviorSet
+ * that they can provide (to UInputRouter, primarily)
+ * @todo callback/delegate for when the provided InputBehaviorSet changes
+ */
+class IInputBehaviorSource
+{
+	GENERATED_BODY()
+
+public:
+	/**
+	 * @return The current UInputBehaviorSet for this Source
+	 */
+	virtual const UInputBehaviorSet* GetInputBehaviors() const = 0;
+};
+
diff --git a/Engine/Source/Runtime/InteractiveToolsFramework/Public/InputRouter.h b/Engine/Source/Runtime/InteractiveToolsFramework/Public/InputRouter.h
new file mode 100644
--- /dev/null
+++ b/Engine/Source/Runtime/InteractiveToolsFramework/Public/InputRouter.h
@@ -0,0 +1,96 @@
+// Copyright 1998-2018 Epic Games, Inc. All Rights Reserved.
+
+#pragma once
+
+#include "CoreMinimal.h"
+#include "UObject/Object.h"
+#include "Misc/Change.h"
+#include "InputState.h"
+#include "InputBehaviorSet.h"
+#include "ToolContextInterfaces.h"
+#include "InputRouter.generated.h"
+
+
+/**
+ * UInputRouter mediates between a higher-level input event source (eg like an FEdMode)
+ * and a set of InputBehaviors that respond to those events. Sets of InputBehaviors are
+ * registered, and then PostInputEvent() is called for each event. 
+ *
+ * Internally one of the active Behaviors may "capture" the event stream.
+ * Separate "Left" and "Right" captures are supported, which means that (eg)
+ * an independent capture can be tracked for each VR controller.
+ *
+ * If the input device supports "hover",  PostHoverInputEvent() will forward
+ * hover events to InputBehaviors that also support it.
+ *
+ */
+UCLASS(Transient)
+class INTERACTIVETOOLSFRAMEWORK_API UInputRouter : public UObject
+{
+	GENERATED_BODY()
+
+protected:
+	friend class UInteractiveToolsContext;		// to call Initialize/Shutdown
+
+	UInputRouter();
+
+	/** Initialize the InputRouter with the necessary Context-level state. UInteractiveToolsContext calls this, you should not. */
+	virtual void Initialize(IToolsContextTransactionsAPI* TransactionsAPI);
+
+	/** Shutdown the InputRouter. Called by UInteractiveToolsContext. */
+	virtual void Shutdown();
+
+public:
+	/** Add a new behavior Source. Behaviors from this source will be added to the active behavior set. */
+	virtual void RegisterSource(IInputBehaviorSource* Source);
+
+	/** Remove Behaviors from this Source from the active set */
+	virtual void DeregisterSource(IInputBehaviorSource* Source);
+
+	/** Insert a new input event which is used to check for new captures, or forwarded to active capture */
+	virtual void PostInputEvent(const FInputDeviceState& Input);
+
+	/** Returns true if there is an active mouse capture */
+	virtual bool HasActiveMouseCapture() const;
+
+	// TODO: other capture queries
+
+	/** Insert a new hover input event which is forwarded to all hover-enabled Behaviors */
+	virtual void PostHoverInputEvent(const FInputDeviceState& Input);
+
+	/** If this Behavior is capturing, call ForceEndCapture() to notify that we are taking capture away */
+	virtual void ForceTerminateSource(IInputBehaviorSource* Source);
+
+	/** Terminate any active captures and end all hovers */
+	virtual void ForceTerminateAll();
+
+
+public:
+	/** If true, then we post an Invalidation (ie redraw) request if any active InputBehavior responds to Hover events (default false) */
+	UPROPERTY()
+	bool bAutoInvalidateOnHover;
+
+	/** If true, then we post an Invalidation (ie redraw) request on every captured input event (default false) */
+	UPROPERTY()
+	bool bAutoInvalidateOnCapture;
+
+
+protected:
+	IToolsContextTransactionsAPI* TransactionsAPI;
+
+	UPROPERTY()
+	UInputBehaviorSet* ActiveInputBehaviors;
+
+	UInputBehavior* ActiveLeftCapture;
+	void* ActiveLeftCaptureOwner;
+	FInputCaptureData ActiveLeftCaptureData;
+
+	UInputBehavior* ActiveRightCapture;
+	void* ActiveRightCaptureOwner;
+	FInputCaptureData ActiveRightCaptureData;
+
+	FInputDeviceState LastHoverInput;
+
+	void CheckForMouseCaptures(const FInputDeviceState& Input);
+	void HandleCapturedMouseInput(const FInputDeviceState& Input);
+};
\ No newline at end of file
diff --git a/Engine/Source/Runtime/InteractiveToolsFramework/Public/InputState.h b/Engine/Source/Runtime/InteractiveToolsFramework/Public/InputState.h
new file mode 100644
--- /dev/null
+++ b/Engine/Source/Runtime/InteractiveToolsFramework/Public/InputState.h
@@ -0,0 +1,156 @@
+// Copyright 1998-2019 Epic Games, Inc. All Rights Reserved.
+
+#pragma once
+
+#include "CoreMinimal.h"
+#include "UObject/ObjectMacros.h"
+#include "Math/UnrealMath.h"
+
+
+/**
+ * Input event data can be applicable to many possible input devices.
+ * These flags are used to indicate specific or sets of device types.
+ */
+UENUM()
+enum class EInputDevices
+{
+	None = 0,
+	Mouse = 1,
+	Gamepad = 2,
+
+	OculusTouch = 4,
+	HTCViveWands = 8,
+	AnySpatialDevice = OculusTouch | HTCViveWands,
+
+	TabletFingers = 1024
+};
+ENUM_CLASS_FLAGS(EInputDevices);
+
+
+
+
+/**
+ * Current State of a physical device button (mouse, key, etc) at a point in time.
+ * Each "click" of a button should involve at minimum two separate state
+ * events, one where bPressed=true and one where bReleased=true.
+ * Each of these states should occur only once.
+ * In addition there may be additional frames where the button is
+ * held down and bDown=true and bPressed=false.
+ */
+struct FDeviceButtonState
+{
+	/** Was the button pressed down this frame. This should happen once per "click" */
+	bool bPressed;
+	/** Is the button currently pressed down. This should be true every frame the button is pressed. */
+	bool bDown;
+	/** Was the button released this frame. This should happen once per "click" */
+	bool bReleased;
+
+	FDeviceButtonState()
+	{
+		bPressed = bDown = bReleased = false;
+	}
+
+	/** Update the states of this button */
+	void SetStates(bool bPressedIn, bool bDownIn, bool bReleasedIn)
+	{
+		bPressed = bPressedIn;
+		bDown = bDownIn;
+		bReleased = bReleasedIn;
+	}
+};
+
+
+
+/**
+ * Current State of a physical Mouse device at a point in time.
+ */
+struct FMouseInputDeviceState
+{
+	/** State of the left mouse button */
+	FDeviceButtonState Left;
+	/** State of the middle mouse button */
+	FDeviceButtonState Middle;
+	/** State of the right mouse button */
+	FDeviceButtonState Right;
+
+	/** Change in 'ticks' of the mouse wheel since last state event */
+	float WheelDelta;
+
+	/** Current 2D position of the mouse, in application-defined coordinate system */
+	FVector2D Position2D;
+
+	/** Change in 2D mouse position from last state event */
+	FVector2D Delta2D;
+
+	/** Ray into current 3D scene at current 2D mouse position */
+	FRay WorldRay;
+
+
+	FMouseInputDeviceState()
+	{
+		Left = FDeviceButtonState();
+		Middle = FDeviceButtonState();
+		Right = FDeviceButtonState();
+		WheelDelta = false;
+		Position2D = FVector2D::ZeroVector;
+		Delta2D = FVector2D::ZeroVector;
+		WorldRay = FRay();
+	}
+};
+
+
+
+/**
+ * Current state of physical input devices at a point in time.
+ * Assumption is that the state refers to a single physical input device,
+ * ie InputDevice field is a single value of EInputDevices and not a combination.
+ */
+struct FInputDeviceState
+{
+	/** Which InputDevice member is valid in this state */
+	EInputDevices InputDevice;
+
+	//
+	// keyboard modifiers
+	//
+
+	/** Is they keyboard SHIFT modifier key currently pressed down */
+	bool bShiftKeyDown;
+	/** Is they keyboard ALT modifier key currently pressed down */
+	bool bAltKeyDown;
+	/** Is they keyboard CTRL modifier key currently pressed down */
+	bool bCtrlKeyDown;
+	/** Is they keyboard CMD modifier key currently pressed down (only on Apple devices) */
+	bool bCmdKeyDown;
+
+	/** Current state of Mouse device, if InputDevice == EInputDevices::Mouse */
+	FMouseInputDeviceState Mouse;
+
+
+	FInputDeviceState() 
+	{
+		InputDevice = EInputDevices::None;
+		bShiftKeyDown = bAltKeyDown = bCtrlKeyDown = bCmdKeyDown = false;
+		Mouse = FMouseInputDeviceState();
+	}
+
+	/** Update keyboard modifier key states */
+	void SetKeyStates(bool bShiftDown, bool bAltDown, bool bCtrlDown, bool bCmdDown) 
+	{
+		bShiftKeyDown = bShiftDown;
+		bAltKeyDown = bAltDown;
+		bCtrlKeyDown = bCtrlDown;
+		bCmdKeyDown = bCmdDown;
+	}
+
+	/**
+	 * @param DeviceType Combination of device-type flags
+	 * @return true if this input state is for an input device that matches the query flags 
+	 */
+	bool IsFromDevice(EInputDevices DeviceType) const
+	{
+		return ((InputDevice & DeviceType) != EInputDevices::None);
+	}
+};
+
diff --git a/Engine/Source/Runtime/InteractiveToolsFramework/Public/InteractiveTool.h b/Engine/Source/Runtime/InteractiveToolsFramework/Public/InteractiveTool.h
new file mode 100644
--- /dev/null
+++ b/Engine/Source/Runtime/InteractiveToolsFramework/Public/InteractiveTool.h
@@ -0,0 +1,143 @@
+// Copyright 1998-2019 Epic Games, Inc. All Rights Reserved.
+
+#pragma once
+
+#include "CoreMinimal.h"
+#include "InputBehaviorSet.h"
+#include "ToolContextInterfaces.h"
+#include "InteractiveTool.generated.h"
+
+class UInteractiveToolManager;
+
+
+/** Passed to UInteractiveTool::Shutdown to indicate how Tool should shut itself down*/
+enum class EToolShutdownType
+{
+	/** Tool cleans up and exits. Pass this to tools that do not have Accept/Cancel options. */
+	Completed = 0,
+	/** Tool commits current preview to scene */
+	Accept = 1,
+	/** Tool discards current preview without modifying scene */
+	Cancel = 2
+};
+
+
+/**
+ * UInteractiveTool is the base class for all Tools in the InteractiveToolsFramework.
+ * A Tool is is a "lightweight mode" that may "own" one or more Actors/Components/etc in
+ * the current scene, may capture certain input devices or event streams, and so on.
+ * The base implementation essentially does nothing but provide sane default behaviors.
+ *
+ * The BaseTools/ subfolder contains implementations of various kinds of standard
+ * "tool behavior", like a tool that responds to a mouse click, etc, that can be
+ * extended to implement custom behaviors.
+ *
+ * In the framework, you do not create instances of UInteractiveTool yourself. 
+ * You provide a UInteractiveToolBuilder implementation that can properly construct
+ * an instance of your Tool, this is where for example default parameters would be set.
+ * The ToolBuilder is registered with the ToolManager, and then UInteractiveToolManager::ActivateTool()
+ * is used to kick things off.
+ *
+ * @todo callback/delegate for if/when .InputBehaviors changes
+ * @todo callback/delegate for when tool properties change
+ */
+UCLASS(Transient)
+class INTERACTIVETOOLSFRAMEWORK_API UInteractiveTool : public UObject, public IInputBehaviorSource
+{
+	GENERATED_BODY()
+
+public:
+	UInteractiveTool();
+
+	/**
+	 * Called by ToolManager to initialize the Tool *after* ToolBuilder::BuildTool() has been called
+	 */
+	virtual void Setup();
+
+	/**
+	 * Called by ToolManager to shut down the Tool
+	 * @param ShutdownType indicates how the tool should shutdown (ie Accept or Cancel current preview, etc)
+	 */
+	virtual void Shutdown(EToolShutdownType ShutdownType);
+
+	/**
+	 * Allow the Tool to do any custom drawing (ie via PDI/RHI)
+	 * @param RenderAPI Abstraction that provides access to Rendering in the current ToolsContext
+	 */
+	virtual void Render(IToolsContextRenderAPI* RenderAPI);
+
+	/**
+	 * Allow the Tool to do any necessary processing on Tick 
+	 * @param DeltaTime the time delta since last tick
+	 */
+	virtual void Tick(float DeltaTime);
+
+
+
+	/**
+	 * @return ToolManager that owns this Tool
+	 */
+	virtual UInteractiveToolManager* GetToolManager() const;
+
+
+	/** 
+	 * @return true if this Tool support being Cancelled, ie calling Shutdown(EToolShutdownType::Cancel)  
+	 */
+	virtual bool HasCancel() const;
+
+	/** 
+	 * @return true if this Tool support being Accepted, ie calling Shutdown(EToolShutdownType::Accept)  
+	 */
+	virtual bool HasAccept() const;
+
+	/** 
+	 * @return true if this Tool is currently in a state where it can be Accepted. This may be false if for example there was an error in the Tool. 
+	 */
+	virtual bool CanAccept() const;
+
+
+
+	//
+	// Input Behaviors support
+	//
+
+	/**
+	 * Add an input behavior for this Tool
+	 * @param Behavior behavior to add
+	 */
+	virtual void AddInputBehavior(UInputBehavior* Behavior);
+
+	/**
+	 * @return Current input behavior set.
+	 */
+	virtual const UInputBehaviorSet* GetInputBehaviors() const;
+
+
+	//
+	// Property support
+	//
+
+	/**
+	 * @return list of property UObjects for this tool (ie to add to a DetailsViewPanel, for example)
+	 */
+	virtual const TArray<UObject*>& GetToolProperties() const;
+
+
+
+protected:
+
+	/** The current set of InputBehaviors provided by this Tool */
+	UPROPERTY()
+	UInputBehaviorSet* InputBehaviors;
+
+	/** The current set of Property UObjects provided by this Tool. May contain pointer to itself. */
+	TArray<UObject*> ToolPropertyObjects;
+
+	/**
+	 * Add a Property object for this Tool
+	 * @param Property object to add
+	 */
+	virtual void AddToolPropertySource(UObject* PropertyObject);
+};
+
+
diff --git a/Engine/Source/Runtime/InteractiveToolsFramework/Public/InteractiveToolBuilder.h b/Engine/Source/Runtime/InteractiveToolsFramework/Public/InteractiveToolBuilder.h
new file mode 100644
--- /dev/null
+++ b/Engine/Source/Runtime/InteractiveToolsFramework/Public/InteractiveToolBuilder.h
@@ -0,0 +1,46 @@
+// Copyright 1998-2019 Epic Games, Inc. All Rights Reserved.
+
+#pragma once
+
+#include "CoreMinimal.h"
+#include "UObject/Interface.h"
+#include "InteractiveTool.h"
+#include "InteractiveToolBuilder.generated.h"
+
+
+
+/**
+ * A UInteractiveToolBuilder creates a new instance of an InteractiveTool (basically this is a Factory).
+ * These are registered with the InteractiveToolManager, which calls BuildTool() if CanBuildTool() returns true.
+ * In addition CanBuildTool() will be queried to (for example) enable/disable UI buttons, etc.
+ * This is an abstract base class, you must subclass it in order to create your particular Tool instance
+ */
+UCLASS(Transient, Abstract)
+class INTERACTIVETOOLSFRAMEWORK_API UInteractiveToolBuilder : public UObject
+{
+	GENERATED_BODY()
+
+public:
+
+	/** 
+	 * Check if, given the current scene state, a new instance of this builder's Tool can be created
+	 * @param SceneState the current scene selection state, etc
+	 * @return true if a new Tool instance can be created
+	 */
+	virtual bool CanBuildTool(const FToolBuilderState& SceneState) const
+	{
+		check(false);
+		return false;
+	}
+
+	/** 
+	 * Create a new instance of this builder's Tool
+	 * @param SceneState the current scene selection state, etc
+	 * @return a new instance of the Tool, or nullptr on error/failure
+	 */
+	virtual UInteractiveTool* BuildTool(const FToolBuilderState& SceneState) const
+	{
+		check(false);
+		return nullptr;
+	}
+};
\ No newline at end of file
diff --git a/Engine/Source/Runtime/InteractiveToolsFramework/Public/InteractiveToolManager.h b/Engine/Source/Runtime/InteractiveToolsFramework/Public/InteractiveToolManager.h
new file mode 100644
--- /dev/null
+++ b/Engine/Source/Runtime/InteractiveToolsFramework/Public/InteractiveToolManager.h
@@ -0,0 +1,196 @@
+// Copyright 1998-2018 Epic Games, Inc. All Rights Reserved.
+
+#pragma once
+
+#include "CoreMinimal.h"
+#include "UObject/Object.h"
+#include "Misc/Change.h"
+#include "InteractiveTool.h"
+#include "InteractiveToolBuilder.h"
+#include "InputRouter.h"
+#include "ToolContextInterfaces.h"
+#include "InteractiveToolManager.generated.h"
+
+
+
+/** A Tool can be activated on a particular input device, currently identified by a "side" */
+UENUM()
+enum class EToolSide
+{
+	/** Left-hand Tool, also used for Mouse */
+	Left = 1,
+	/** Right-hand Tool*/
+	Right = 2
+};
+
+
+/**
+ * UInteractiveToolManager allows users of the tools framework to create and operate Tool instances.
+ * For each Tool, a (string,ToolBuilder) pair is registered with the ToolManager.
+ * Tools can then be activated via the string identifier.
+ * 
+ * Currently a single Tool can be active for each input device. So for mouse input a single
+ * Tool is available and effectively a lightweight mode. The mouse uses the "Left" tool slot.
+ * 
+ * For VR controllers and touch input, a "Left" and "Right" tool can be active at the same time.  
+ * @todo this is not fully supported yet
+ * 
+ * Tools are not directly created. Use SelectActiveToolType(side,string) to set the active ToolBuilder
+ * on a given side, and then use ActivateTool() to create the new Tool instance.
+ *
+ */
+UCLASS(Transient)
+class INTERACTIVETOOLSFRAMEWORK_API UInteractiveToolManager : public UObject
+{
+	GENERATED_BODY()
+
+protected:
+	friend class UInteractiveToolsContext;		// to call Initialize/Shutdown
+
+	UInteractiveToolManager();
+
+	/** Initialize the ToolManager with the necessary Context-level state. UInteractiveToolsContext calls this, you should not. */
+	virtual void Initialize(IToolsContextQueriesAPI* QueriesAPI, IToolsContextTransactionsAPI* TransactionsAPI, UInputRouter* InputRouter);
+
+	/** Shutdown the ToolManager. Called by UInteractiveToolsContext. */
+	virtual void Shutdown();
+
+public:
+
+	//
+	// Tool registration and Current Tool state
+	//
+
+	/**
+	 * Register a new ToolBuilder
+	 * @param Identifier string used to identify this Builder
+	 * @param Builder new ToolBuilder instance
+	 */
+	virtual void RegisterToolType(const FString& Identifier, UInteractiveToolBuilder* Builder);
+
+	/**
+	 * Set active ToolBuilder for a ToolSide via string identifier
+	 * @param Side which "side" should we set this Builder on
+	 * @param Identifier name of ToolBuilder that was passed to RegisterToolType()
+	 */
+	virtual bool SelectActiveToolType(EToolSide Side, const FString& Identifier);
+
+	/**
+	 * Check if a named Tool type can currently be activated on the given ToolSide
+	 * @param Side which "side" you would like to active the tool on
+	 * @param Identifier string name of the Tool type
+	 * @return true if the Tool type could be activated
+	 */
+	virtual bool CanActivateTool(EToolSide eSide, const FString& Identifier);
+
+	/**
+	 * Try to activate a new Tool instance on the given Side
+	 * @param Side which "side" you would like to active the tool on
+	 * @return true if a new Tool instance was created and initialized
+	 */	
+	virtual bool ActivateTool(EToolSide Side);
+
+	/**
+	 * Check if there is an active Tool on the given Side
+	 * @param Side which Side to check
+	 * @return true if there is an active Tool on that side
+	 */
+	virtual bool HasActiveTool(EToolSide Side);
+
+	/**
+	 * Get pointer to active Tool on a given side
+	 * @param Side which Side is being requested
+	 * @return pointer to Tool instance active on that Side, or nullptr if no such Tool exists
+	 */
+	virtual UInteractiveTool* GetActiveTool(EToolSide Side);
+
+
+	/**
+	 * Check if an active Tool on the given Side can be Accepted in its current state
+	 * @param Side which Side to check
+	 * @return true if there is an active Tool and it returns true from HasAccept() and CanAccept()
+	 */
+	virtual bool CanAcceptActiveTool(EToolSide Side);
+
+	/**
+	 * Check if an active Tool on the given Side can be Canceled
+	 * @param Side which Side to check
+	 * @return true if there is an active Tool and it returns true from HasCancel()
+	 */
+	virtual bool CanCancelActiveTool(EToolSide Side);
+
+	/**
+	 * Shut down an active Tool on the given side
+	 * @param Side which "side" you would like to shut down
+	 * @param ShutdownType how should the tool be terminated (eg Accept/Cancel)
+	 */
+	virtual void DeactivateTool(EToolSide Side, EToolShutdownType ShutdownType);
+
+
+
+	//
+	// Functions that Tools can call to interact with Transactions API
+	//
+	
+	
+	/** Request an Invalidation via the Transactions API (ie to cause a repaint, etc) */
+	virtual void PostInvalidation();
+
+	/**
+	 * Request that the Context open a Transaction, whatever that means to the current Context
+	 * @param Description text description of this transaction (this is the string that appears on undo/redo in the UE Editor)
+	 */
+	virtual void BeginUndoTransaction(const FText& Description);
+
+	/** Request that the Context close and commit the open Transaction */
+	virtual void EndUndoTransaction();
+
+	/**
+	 * Forward an FChange object to the Context
+	 * @param TargetObject the object that the FChange applies to
+	 * @param Change the change object that the Context should insert into the transaction history
+	 * @param Description text description of this change (this is the string that appears on undo/redo in the UE Editor)
+	 */
+	virtual void EmitObjectChange(UObject* TargetObject, TUniquePtr<FChange> Change, const FText& Description );
+
+
+	//
+	// State control  (@todo: have the Context call these? not safe for anyone to call)
+	//
+
+	/** Tick any active Tools. Called by UInteractiveToolsContext */
+	virtual void Tick(float DeltaTime);
+
+	/** Render any active Tools. Called by UInteractiveToolsContext. */
+	virtual void Render(IToolsContextRenderAPI* RenderAPI);
+
+
+public:
+	/** Currently-active Left Tool, or null if no Tool is active */
+	UPROPERTY()
+	UInteractiveTool* ActiveLeftTool;
+
+	/** Currently-active Right Tool, or null if no Tool is active */
+	UPROPERTY()
+	UInteractiveTool* ActiveRightTool;
+
+
+protected:
+	/** Pointer to current Context-Queries implementation */
+	IToolsContextQueriesAPI* QueriesAPI;
+	/** Pointer to current Transactions implementation */
+	IToolsContextTransactionsAPI* TransactionsAPI;
+
+	/** Pointer to current InputRouter (Context owns this) */
+	UInputRouter* InputRouter;
+
+	/** Current set of named ToolBuilders */
+	UPROPERTY()
+	TMap<FString, UInteractiveToolBuilder*> ToolBuilders;
+
+	/** Currently-active Left ToolBuilder */
+	UInteractiveToolBuilder* ActiveLeftBuilder;
+	/** Currently-active Right ToolBuilder */
+	UInteractiveToolBuilder* ActiveRightBuilder;
+
+};
\ No newline at end of file
diff --git a/Engine/Source/Runtime/InteractiveToolsFramework/Public/InteractiveToolsContext.h b/Engine/Source/Runtime/InteractiveToolsFramework/Public/InteractiveToolsContext.h
new file mode 100644
--- /dev/null
+++ b/Engine/Source/Runtime/InteractiveToolsFramework/Public/InteractiveToolsContext.h
@@ -0,0 +1,41 @@
+// Copyright 1998-2018 Epic Games, Inc. All Rights Reserved.
+
+#pragma once
+
+#include "CoreMinimal.h"
+#include "InteractiveToolManager.h"
+#include "InteractiveToolsContext.generated.h"
+
+/**
+ * InteractiveToolsContext owns a ToolManager and an InputRouter. This is just a top-level 
+ * UObject container, however implementations like UEdModeInteractiveToolsContext extend
+ * this class to make it easier to connect external systems (like an FEdMode) to the ToolsFramework.
+ */
+UCLASS(Transient)
+class INTERACTIVETOOLSFRAMEWORK_API UInteractiveToolsContext : public UObject
+{
+	GENERATED_BODY()
+	
+public:
+	UInteractiveToolsContext();
+
+	/** 
+	 * Initialize the Context. This creates the InputRouter and ToolManager 
+	 * @param QueriesAPI client-provided implementation of the API for querying the higher-evel scene state
+	 * @param TransactionsAPI client-provided implementation of the API for publishing events and transactions
+	 */
+	virtual void Initialize(IToolsContextQueriesAPI* QueriesAPI, IToolsContextTransactionsAPI* TransactionsAPI);
+
+	/** Shutdown Context by destroying InputRouter and ToolManager */
+	virtual void Shutdown();
+
+public:
+	/** current UInputRouter for this Context */
+	UPROPERTY()
+	UInputRouter* InputRouter;	
+
+	/** current UInteractiveToolManager for this Context */
+	UPROPERTY()
+	UInteractiveToolManager* ToolManager;	
+
+};
\ No newline at end of file
diff --git a/Engine/Source/Runtime/InteractiveToolsFramework/Public/InteractiveToolsFramework.h b/Engine/Source/Runtime/InteractiveToolsFramework/Public/InteractiveToolsFramework.h
new file mode 100644
--- /dev/null
+++ b/Engine/Source/Runtime/InteractiveToolsFramework/Public/InteractiveToolsFramework.h
@@ -0,0 +1,15 @@
+// Copyright 1998-2018 Epic Games, Inc. All Rights Reserved.
+
+#pragma once
+
+#include "CoreMinimal.h"
+#include "Modules/ModuleManager.h"
+
+class FInteractiveToolsFrameworkModule : public IModuleInterface
+{
+public:
+
+	/** IModuleInterface implementation */
+	virtual void StartupModule() override;
+	virtual void ShutdownModule() override;
+};
diff --git a/Engine/Source/Runtime/InteractiveToolsFramework/Public/ToolBuilderUtil.h b/Engine/Source/Runtime/InteractiveToolsFramework/Public/ToolBuilderUtil.h
new file mode 100644
--- /dev/null
+++ b/Engine/Source/Runtime/InteractiveToolsFramework/Public/ToolBuilderUtil.h
@@ -0,0 +1,113 @@
+// Copyright 1998-2018 Epic Games, Inc. All Rights Reserved.
+
+#pragma once
+
+#include "CoreMinimal.h"
+#include "GameFramework/Actor.h"
+#include "Engine/Selection.h"
+#include "InteractiveToolBuilder.h"
+
+/**
+* Helper functions that can be used in InteractiveToolBuilder implementations
+*/
+namespace ToolBuilderUtil
+{
+	/** Returns true if this UObject can provide a FMeshDescription */
+	INTERACTIVETOOLSFRAMEWORK_API
+	bool IsMeshDescriptionSourceComponent(UObject* ComponentObject);
+	
+	/** Count number of selected components that pass predicate. If Component selection is not empty, returns that count, otherwise counts in all selected Actors */
+	INTERACTIVETOOLSFRAMEWORK_API
+	int CountComponents(const FToolBuilderState& InputState, const TFunction<bool(UObject*)>& Predicate);
+
+	/** First first available component that passes predicate. Searches Components selection list first, then all Actors */
+	INTERACTIVETOOLSFRAMEWORK_API
+	UObject* FindFirstComponent(const FToolBuilderState& InputState, const TFunction<bool(UObject*)>& Predicate);
+
+	// @todo not sure that actors with multiple components are handled properly...
+	/** Count number of components of given type. If Component selection is not empty, returns that count, otherwise counts in all selected Actors */
+	template<typename ComponentType>
+	int CountSelectedComponentsOfType(const FToolBuilderState& InputState);
+
+	/** First first available component of given type. Searches Components selection list first, then all Actors */
+	template<typename ComponentType>
+	ComponentType* FindFirstComponentOfType(const FToolBuilderState& InputState);
+
+}
+
+
+
+
+/*
+ * Template Implementations
+ */
+
+
+template<typename ComponentType>
+int ToolBuilderUtil::CountSelectedComponentsOfType(const FToolBuilderState& InputState)
+{
+	int nTypedComponents = 0;
+
+	if (InputState.SelectedComponents != nullptr && InputState.SelectedComponents->Num() > 0)
+	{
+		for (FSelectionIterator Iter(*InputState.SelectedComponents); Iter; ++Iter)
+		{
+			if (ComponentType* TypedComponent = Cast<ComponentType>(*Iter))
+			{
+				nTypedComponents++;
+			}
+		}
+	}
+	else
+	{
+		for (FSelectionIterator Iter(*InputState.SelectedActors); Iter; ++Iter)
+		{
+			if (AActor* Actor = Cast<AActor>(*Iter))
+			{
+				// [RMS] is there a more efficient way to do this?
+				TArray<UActorComponent*> TypedComponents = Actor->GetComponentsByClass(ComponentType::StaticClass());
+				nTypedComponents += TypedComponents.Num();
+			}
+		}
+	}
+
+	return nTypedComponents;
+}
+
+
+
+
+
+
+template<typename ComponentType>
+ComponentType* ToolBuilderUtil::FindFirstComponentOfType(const FToolBuilderState& InputState)
+{
+	if (InputState.SelectedComponents != nullptr && InputState.SelectedComponents->Num() > 0)
+	{
+		for (FSelectionIterator Iter(*InputState.SelectedComponents); Iter; ++Iter)
+		{
+			if (ComponentType* TypedComponent = Cast<ComponentType>(*Iter))
+			{
+				return TypedComponent;
+			}
+		}
+	}
+	else
+	{
+		for (FSelectionIterator Iter(*InputState.SelectedActors); Iter; ++Iter)
+		{
+			if (AActor* Actor = Cast<AActor>(*Iter))
+			{
+				// [RMS] is there a more efficient way to do this?
+				TArray<UActorComponent*> TypedComponents = Actor->GetComponentsByClass(ComponentType::StaticClass());
+				if (TypedComponents.Num() > 0)
+				{
+					return (ComponentType*)TypedComponents[0];
+				}
+			}
+		}
+	}
+
+	return nullptr;
+}
+
diff --git a/Engine/Source/Runtime/InteractiveToolsFramework/Public/ToolContextInterfaces.h b/Engine/Source/Runtime/InteractiveToolsFramework/Public/ToolContextInterfaces.h
new file mode 100644
--- /dev/null
+++ b/Engine/Source/Runtime/InteractiveToolsFramework/Public/ToolContextInterfaces.h
@@ -0,0 +1,174 @@
+// Copyright 1998-2019 Epic Games, Inc. All Rights Reserved.
+
+#pragma once
+
+#include "CoreMinimal.h"
+#include "ComponentSourceInterfaces.h"
+
+
+// predeclarations so we don't have to include these in all tools
+class FChange;
+class UPackage;
+class FPrimitiveDrawInterface;
+class USelection;
+class UInteractiveToolManager;
+
+
+/**
+ * FToolBuilderState is a bucket of state information that a ToolBuilder might need
+ * to construct a Tool. This information comes from a level above the Tools framework,
+ * and depends on the context we are in (Editor vs Runtime, for example).
+ */
+struct INTERACTIVETOOLSFRAMEWORK_API FToolBuilderState
+{
+	/** The current UWorld */
+	UWorld* World;
+	/** The current ToolManager */
+	UInteractiveToolManager* ToolManager;
+
+	/** Current selected Actors. May be empty or nullptr. */
+	USelection* SelectedActors;
+	/** Current selected Components. May be empty or nullptr. */
+	USelection* SelectedComponents;
+
+	/** Implementation that can build Sources (like MeshDescriptionSource) for Components */
+	IComponentSourceFactory* SourceBuilder;
+
+	FToolBuilderState()
+	{
+		World = nullptr;
+		ToolManager = nullptr;
+		SelectedActors = nullptr;
+		SelectedComponents = nullptr;
+		SourceBuilder = nullptr;
+	}
+};
+
+
+
+
+/**
+ * Users of the Tools Framework need to implement IToolsContextQueriesAPI to provide
+ * access to scene state information like the current UWorld, active USelections, etc.
+ */
+class IToolsContextQueriesAPI
+{
+public:
+
+	/**
+	 * Collect up current-selection information for the current scene state (ie what is selected in Editor, etc)
+	 * @param StateOut this structure is populated with available state information
+	 */
+	virtual void GetCurrentSelectionState(FToolBuilderState& StateOut) const = 0;
+};
+
+
+
+
+/** Level of severity of messages emitted by Tool framework */
+UENUM()
+enum class EToolMessageLevel
+{
+	/** Development message goes into development log */
+	Internal = 0,
+	/** User message should appear in user-facing log */
+	UserMessage = 1,
+	/** Notification message should be shown in a non-modal notification window */
+	UserNotification = 2,
+	/** Error message should be shown in a modal notification window */
+	UserError = 3
+};
+
+
+/**
+ * Users of the Tools Framework need to implement IToolsContextTransactionsAPI so that
+ * the Tools have the ability to create Transactions and emit Changes. Note that this is
+ * technically optional, but that undo/redo won't be supported without it.
+ */
+class IToolsContextTransactionsAPI
+{
+public:
+
+	/**
+	 * Request that context display message information.
+	 * @param Message text of message
+	 * @param Level severity level of message
+	 */
+	virtual void PostMessage(const TCHAR* Message, EToolMessageLevel Level) = 0;
+
+	/** 
+	 * Forward an invalidation request from Tools framework, to cause repaint/etc. 
+	 * This is not always necessary but in some situations (eg in Non-Realtime mode in Editor)
+	 * a redraw will not happen every frame. 
+	 * See UInputRouter for options to enable auto-invalidation.
+	 */
+	virtual void PostInvalidation() = 0;
+	
+	/**
+	 * Begin a Transaction, whatever this means in the current Context. For example in the
+	 * Editor it means open a GEditor Transaction. You must call EndUndoTransaction() after calling this.
+	 * @param Description text description of the transaction that could be shown to user
+	 */
+	virtual void BeginUndoTransaction(const FText& Description) = 0;
+
+	/**
+	 * Complete the Transaction. Assumption is that Begin/End are called in pairs.
+	 */
+	virtual void EndUndoTransaction() = 0;
+
+	/**
+	 * Insert an FChange into the transaction history in the current Context. 
+	 * This cannot be called between Begin/EndUndoTransaction, the FChange should be 
+	 * automatically inserted into a Transaction.
+	 * @param TargetObject The UObject this Change is applied to
+	 * @param Change The Change implementation
+	 * @param Description text description of the transaction that could be shown to user
+	 */
+	virtual void AppendChange(UObject* TargetObject, TUniquePtr<FChange> Change, const FText& Description) = 0;
+};
+
+
+/**
+ * Users of the Tools Framework need to implement IToolsContextRenderAPI to allow
+ * Tools, Indicators, and Gizmos to make low-level rendering calls for things like line drawing.
+ * This API will be passed to eg UInteractiveTool::Render(), so access is only provided when it
+ * makes sense to call the functions
+ */
+class IToolsContextRenderAPI
+{
+public:
+	/** @return Current PDI */
+	virtual FPrimitiveDrawInterface* GetPrimitiveDrawInterface() = 0;
+};
+
+
+
+
+/**
+ * Users of the Tools Framework need to provide an IToolsContextAssetAPI implementation
+ * that allows Packages and Assets to be created/saved. Note that this is not strictly
+ * necessary, for example a trivial implementation could just store things in the Transient
+ * package and not do any saving.
+ */
+class IToolsContextAssetAPI
+{
+public:
+	/** Get default path to save assets in. For example the currently-visible path in the Editor. */
+	virtual FString GetActiveAssetFolderPath() = 0;
+
+	/** Combines folder and asset names */
+	virtual FString MakePackageName(const FString& AssetName, const FString& FolderPath) = 0;
+
+	/** return "unique" version of AssetName, created by appending _1, _2, ... if AssetName already exists */
+	virtual FString MakeUniqueAssetName(const FString& AssetName, const FString& FolderPath) = 0;
+
+	/** create a new package at the given path (calls MakePackageName) */
+	virtual UPackage* CreateNewPackage(const FString& AssetName, const FString& FolderPath) = 0;
+
+	/** Request saving of asset to persistent storage via something like an interactive popup dialog */
+	virtual void InteractiveSaveGeneratedAsset(UObject* Asset, UPackage* AssetPackage) = 0;
+
+	/** Autosave asset to persistent storage */
+	virtual void AutoSaveGeneratedAsset(UObject* Asset, UPackage* AssetPackage) = 0;
+};
+
