fix copyright year and static analysis

diff --git a/Engine/Source/Editor/EditorInteractiveToolsFramework/Private/EditorInteractiveToolsFrameworkModule.cpp b/Engine/Source/Editor/EditorInteractiveToolsFramework/Private/EditorInteractiveToolsFrameworkModule.cpp
--- a/Engine/Source/Editor/EditorInteractiveToolsFramework/Private/EditorInteractiveToolsFrameworkModule.cpp
+++ b/Engine/Source/Editor/EditorInteractiveToolsFramework/Private/EditorInteractiveToolsFrameworkModule.cpp
@@ -1,4 +1,4 @@
-// Copyright 1998-2018 Epic Games, Inc. All Rights Reserved.
+// Copyright 1998-2019 Epic Games, Inc. All Rights Reserved.
 
 #include "EditorInteractiveToolsFrameworkModule.h"
 
diff --git a/Engine/Source/Editor/EditorInteractiveToolsFramework/Public/EdModeInteractiveToolsContext.h b/Engine/Source/Editor/EditorInteractiveToolsFramework/Public/EdModeInteractiveToolsContext.h
--- a/Engine/Source/Editor/EditorInteractiveToolsFramework/Public/EdModeInteractiveToolsContext.h
+++ b/Engine/Source/Editor/EditorInteractiveToolsFramework/Public/EdModeInteractiveToolsContext.h
@@ -1,4 +1,4 @@
-// Copyright 1998-2018 Epic Games, Inc. All Rights Reserved.
+// Copyright 1998-2019 Epic Games, Inc. All Rights Reserved.
 
 #pragma once
 
diff --git a/Engine/Source/Editor/EditorInteractiveToolsFramework/Public/EditorInteractiveToolsFrameworkModule.h b/Engine/Source/Editor/EditorInteractiveToolsFramework/Public/EditorInteractiveToolsFrameworkModule.h
--- a/Engine/Source/Editor/EditorInteractiveToolsFramework/Public/EditorInteractiveToolsFrameworkModule.h
+++ b/Engine/Source/Editor/EditorInteractiveToolsFramework/Public/EditorInteractiveToolsFrameworkModule.h
@@ -1,4 +1,4 @@
-// Copyright 1998-2018 Epic Games, Inc. All Rights Reserved.
+// Copyright 1998-2019 Epic Games, Inc. All Rights Reserved.
 
 #pragma once
 
diff --git a/Engine/Source/Runtime/InteractiveToolsFramework/InteractiveToolsFramework.Build.cs b/Engine/Source/Runtime/InteractiveToolsFramework/InteractiveToolsFramework.Build.cs
--- a/Engine/Source/Runtime/InteractiveToolsFramework/InteractiveToolsFramework.Build.cs
+++ b/Engine/Source/Runtime/InteractiveToolsFramework/InteractiveToolsFramework.Build.cs
@@ -1,4 +1,4 @@
-// Copyright 1998-2018 Epic Games, Inc. All Rights Reserved.
+// Copyright 1998-2019 Epic Games, Inc. All Rights Reserved.
 
 using UnrealBuildTool;
 
diff --git a/Engine/Source/Runtime/InteractiveToolsFramework/Private/InteractiveToolsFramework.cpp b/Engine/Source/Runtime/InteractiveToolsFramework/Private/InteractiveToolsFramework.cpp
--- a/Engine/Source/Runtime/InteractiveToolsFramework/Private/InteractiveToolsFramework.cpp
+++ b/Engine/Source/Runtime/InteractiveToolsFramework/Private/InteractiveToolsFramework.cpp
@@ -1,4 +1,4 @@
-// Copyright 1998-2018 Epic Games, Inc. All Rights Reserved.
+// Copyright 1998-2019 Epic Games, Inc. All Rights Reserved.
 
 #include "InteractiveToolsFramework.h"
 
diff --git a/Engine/Source/Runtime/InteractiveToolsFramework/Public/InputRouter.h b/Engine/Source/Runtime/InteractiveToolsFramework/Public/InputRouter.h
--- a/Engine/Source/Runtime/InteractiveToolsFramework/Public/InputRouter.h
+++ b/Engine/Source/Runtime/InteractiveToolsFramework/Public/InputRouter.h
@@ -1,4 +1,4 @@
-// Copyright 1998-2018 Epic Games, Inc. All Rights Reserved.
+// Copyright 1998-2019 Epic Games, Inc. All Rights Reserved.
 
 #pragma once
 
diff --git a/Engine/Source/Runtime/InteractiveToolsFramework/Public/InteractiveToolManager.h b/Engine/Source/Runtime/InteractiveToolsFramework/Public/InteractiveToolManager.h
--- a/Engine/Source/Runtime/InteractiveToolsFramework/Public/InteractiveToolManager.h
+++ b/Engine/Source/Runtime/InteractiveToolsFramework/Public/InteractiveToolManager.h
@@ -1,4 +1,4 @@
-// Copyright 1998-2018 Epic Games, Inc. All Rights Reserved.
+// Copyright 1998-2019 Epic Games, Inc. All Rights Reserved.
 
 #pragma once
 
diff --git a/Engine/Source/Runtime/InteractiveToolsFramework/Public/InteractiveToolsContext.h b/Engine/Source/Runtime/InteractiveToolsFramework/Public/InteractiveToolsContext.h
--- a/Engine/Source/Runtime/InteractiveToolsFramework/Public/InteractiveToolsContext.h
+++ b/Engine/Source/Runtime/InteractiveToolsFramework/Public/InteractiveToolsContext.h
@@ -1,4 +1,4 @@
-// Copyright 1998-2018 Epic Games, Inc. All Rights Reserved.
+// Copyright 1998-2019 Epic Games, Inc. All Rights Reserved.
 
 #pragma once
 
diff --git a/Engine/Source/Runtime/InteractiveToolsFramework/Public/InteractiveToolsFramework.h b/Engine/Source/Runtime/InteractiveToolsFramework/Public/InteractiveToolsFramework.h
--- a/Engine/Source/Runtime/InteractiveToolsFramework/Public/InteractiveToolsFramework.h
+++ b/Engine/Source/Runtime/InteractiveToolsFramework/Public/InteractiveToolsFramework.h
@@ -1,4 +1,4 @@
-// Copyright 1998-2018 Epic Games, Inc. All Rights Reserved.
+// Copyright 1998-2019 Epic Games, Inc. All Rights Reserved.
 
 #pragma once
 
diff --git a/Engine/Source/Runtime/InteractiveToolsFramework/Public/ToolBuilderUtil.h b/Engine/Source/Runtime/InteractiveToolsFramework/Public/ToolBuilderUtil.h
--- a/Engine/Source/Runtime/InteractiveToolsFramework/Public/ToolBuilderUtil.h
+++ b/Engine/Source/Runtime/InteractiveToolsFramework/Public/ToolBuilderUtil.h
@@ -1,4 +1,4 @@
-// Copyright 1998-2018 Epic Games, Inc. All Rights Reserved.
+// Copyright 1998-2019 Epic Games, Inc. All Rights Reserved.
 
 #pragma once
 
diff --git a/Engine/Source/Runtime/InteractiveToolsFramework/Public/ToolContextInterfaces.h b/Engine/Source/Runtime/InteractiveToolsFramework/Public/ToolContextInterfaces.h
--- a/Engine/Source/Runtime/InteractiveToolsFramework/Public/ToolContextInterfaces.h
+++ b/Engine/Source/Runtime/InteractiveToolsFramework/Public/ToolContextInterfaces.h
@@ -54,6 +54,7 @@ struct INTERACTIVETOOLSFRAMEWORK_API FToolBuilderState
 class IToolsContextQueriesAPI
 {
 public:
+	virtual ~IToolsContextQueriesAPI() {}
 
 	/**
 	 * Collect up current-selection information for the current scene state (ie what is selected in Editor, etc)
