Add hotkeys to cancel (ESC) and accept/end (ENTER) current tool
#rnx

diff --git a/Engine/Source/Editor/EditorInteractiveToolsFramework/Private/EdModeInteractiveToolsContext.cpp b/Engine/Source/Editor/EditorInteractiveToolsFramework/Private/EdModeInteractiveToolsContext.cpp
--- a/Engine/Source/Editor/EditorInteractiveToolsFramework/Private/EdModeInteractiveToolsContext.cpp
+++ b/Engine/Source/Editor/EditorInteractiveToolsFramework/Private/EdModeInteractiveToolsContext.cpp
@@ -223,6 +223,41 @@ bool UEdModeInteractiveToolsContext::InputKey(FEditorViewportClient* ViewportCli
 
 	bool bHandled = false;
 
+
+	// escape key cancels current tool
+	// @todo also cancels current selection, and only works if viewport is focused :(
+	if (Key == EKeys::Escape && Event == IE_Released && ToolManager->HasAnyActiveTool() )
+	{
+		if (ToolManager->HasActiveTool(EToolSide::Mouse))
+		{
+			ToolManager->DeactivateTool(EToolSide::Mouse, EToolShutdownType::Cancel);
+		}
+		return true;
+	}
+
+	// enter key accepts current tool, or ends tool if it does not have accept state
+	if (Key == EKeys::Enter && Event == IE_Released && ToolManager->HasAnyActiveTool())
+	{
+		if (ToolManager->HasActiveTool(EToolSide::Mouse)) 
+		{
+			if (ToolManager->GetActiveTool(EToolSide::Mouse)->HasAccept())
+			{
+				if (ToolManager->CanAcceptActiveTool(EToolSide::Mouse))
+				{
+					ToolManager->DeactivateTool(EToolSide::Mouse, EToolShutdownType::Accept);
+					return true;
+				}
+			}
+			else
+			{
+				ToolManager->DeactivateTool(EToolSide::Mouse, EToolShutdownType::Completed);
+				return true;
+			}
+		}
+	}
+
+
+
 	if (Event == IE_Pressed || Event == IE_Released)
 	{
 		bool bIsLeftMouse = (Key == EKeys::LeftMouseButton);
diff --git a/Engine/Source/Runtime/InteractiveToolsFramework/Private/InteractiveToolManager.cpp b/Engine/Source/Runtime/InteractiveToolsFramework/Private/InteractiveToolManager.cpp
--- a/Engine/Source/Runtime/InteractiveToolsFramework/Private/InteractiveToolManager.cpp
+++ b/Engine/Source/Runtime/InteractiveToolsFramework/Private/InteractiveToolManager.cpp
@@ -153,11 +153,17 @@ void UInteractiveToolManager::DeactivateTool(EToolSide Side, EToolShutdownType S
 }
 
 
-bool UInteractiveToolManager::HasActiveTool(EToolSide Side)
+bool UInteractiveToolManager::HasActiveTool(EToolSide Side) const
 {
 	return (Side == EToolSide::Left) ? (ActiveLeftTool != nullptr) : (ActiveRightTool != nullptr);
 }
 
+bool UInteractiveToolManager::HasAnyActiveTool() const
+{
+	return ActiveLeftTool != nullptr || ActiveRightTool != nullptr;
+}
+
+
 UInteractiveTool* UInteractiveToolManager::GetActiveTool(EToolSide Side)
 {
 	return (Side == EToolSide::Left) ? ActiveLeftTool : ActiveRightTool;
diff --git a/Engine/Source/Runtime/InteractiveToolsFramework/Public/InteractiveToolManager.h b/Engine/Source/Runtime/InteractiveToolsFramework/Public/InteractiveToolManager.h
--- a/Engine/Source/Runtime/InteractiveToolsFramework/Public/InteractiveToolManager.h
+++ b/Engine/Source/Runtime/InteractiveToolsFramework/Public/InteractiveToolManager.h
@@ -19,8 +19,9 @@ enum class EToolSide
 {
 	/** Left-hand Tool, also used for Mouse */
 	Left = 1,
+	Mouse = 1,
 	/** Right-hand Tool*/
-	Right = 2
+	Right = 2,
 };
 
 
@@ -95,7 +96,14 @@ public:
 	 * @param Side which Side to check
 	 * @return true if there is an active Tool on that side
 	 */
-	virtual bool HasActiveTool(EToolSide Side);
+	virtual bool HasActiveTool(EToolSide Side) const;
+
+
+	/**
+	 * @return true if there are any active tools
+	 */
+	virtual bool HasAnyActiveTool() const;
+
 
 	/**
 	 * Get pointer to active Tool on a given side
