ToolsContextQueryAPI can now return "standard" materials for Tools to use. TransactionsAPI supports actor-selection-change requests. Implemented support in EdModeInteractiveToolsContext.
#rnx

diff --git a/Engine/Source/Editor/EditorInteractiveToolsFramework/Private/EdModeInteractiveToolsContext.cpp b/Engine/Source/Editor/EditorInteractiveToolsFramework/Private/EdModeInteractiveToolsContext.cpp
--- a/Engine/Source/Editor/EditorInteractiveToolsFramework/Private/EdModeInteractiveToolsContext.cpp
+++ b/Engine/Source/Editor/EditorInteractiveToolsFramework/Private/EdModeInteractiveToolsContext.cpp
@@ -39,6 +39,16 @@ public:
 		StateOut.SourceBuilder = ToolsContext->GetComponentSourceFactory();
 	}
 
+	virtual UMaterialInterface* GetStandardMaterial(EStandardToolContextMaterials MaterialType) const
+	{
+		if (MaterialType == EStandardToolContextMaterials::VertexColorMaterial)
+		{
+			return ToolsContext->StandardVertexColorMaterial;
+		}
+		check(false);
+		return nullptr;
+	}
+
 };
 
 
@@ -86,6 +96,33 @@ public:
 		// end transaction
 	}
 
+
+	virtual bool RequestSelectionChange(const FSelectedOjectsChangeList& SelectionChange) override
+	{
+		checkf(SelectionChange.Components.Num() == 0, TEXT("FEdModeToolsContextTransactionImpl::RequestSelectionChange - Component selection not supported yet"));
+
+		if (SelectionChange.ModificationType == ESelectedObjectsModificationType::Clear)
+		{
+			GEditor->SelectNone(true, true, false);
+			return true;
+		}
+
+		if (SelectionChange.ModificationType == ESelectedObjectsModificationType::Replace )
+		{
+			GEditor->SelectNone(false, true, false);
+		}
+
+		bool bAdd = (SelectionChange.ModificationType != ESelectedObjectsModificationType::Remove);
+		int NumActors = SelectionChange.Actors.Num();
+		for (int k = 0; k < NumActors; ++k)
+		{
+			GEditor->SelectActor(SelectionChange.Actors[k], bAdd, false, true, false);
+		}
+
+		GEditor->NoteSelectionChange(true);
+		return true;
+	}
+
 };
 
 
@@ -103,9 +140,9 @@ UEdModeInteractiveToolsContext::UEdModeInteractiveToolsContext()
 
 
 
-void UEdModeInteractiveToolsContext::Initialize(IToolsContextQueriesAPI* queriesAPI, IToolsContextTransactionsAPI* transactionsAPI)
+void UEdModeInteractiveToolsContext::Initialize(IToolsContextQueriesAPI* QueriesAPIIn, IToolsContextTransactionsAPI* TransactionsAPIIn)
 {
-	UInteractiveToolsContext::Initialize(queriesAPI, transactionsAPI);
+	UInteractiveToolsContext::Initialize(QueriesAPIIn, TransactionsAPIIn);
 
 	bInvalidationPending = false;
 }
@@ -127,17 +164,21 @@ void UEdModeInteractiveToolsContext::InitializeContextFromEdMode(FEdMode* Editor
 	this->AssetAPI = new FEditorToolAssetAPI();
 	this->SourceFactory = new FEditorComponentSourceFactory();
 
-	UInteractiveToolsContext::Initialize(QueriesAPI, TransactionAPI);
+	Initialize(QueriesAPI, TransactionAPI);
 
 	// enable auto invalidation in Editor, because invalidating for all hover and capture events is unpleasant
 	this->InputRouter->bAutoInvalidateOnHover = true;
 	this->InputRouter->bAutoInvalidateOnCapture = true;
+
+
+	// set up standard materials
+	StandardVertexColorMaterial = LoadObject<UMaterial>(nullptr, TEXT("/Game/Materials/VertexColor"));
 }
 
 
 void UEdModeInteractiveToolsContext::ShutdownContext()
 {
-	UInteractiveToolsContext::Shutdown();
+	Shutdown();
 
 	if (QueriesAPI != nullptr)
 	{
@@ -225,14 +266,16 @@ bool UEdModeInteractiveToolsContext::InputKey(FEditorViewportClient* ViewportCli
 
 
 	// escape key cancels current tool
-	// @todo also cancels current selection, and only works if viewport is focused :(
-	if (Key == EKeys::Escape && Event == IE_Released && ToolManager->HasAnyActiveTool() )
+	if (Key == EKeys::Escape && Event == IE_Released )
 	{
-		if (ToolManager->HasActiveTool(EToolSide::Mouse))
+		if (ToolManager->HasAnyActiveTool())
 		{
-			ToolManager->DeactivateTool(EToolSide::Mouse, EToolShutdownType::Cancel);
+			if (ToolManager->HasActiveTool(EToolSide::Mouse))
+			{
+				ToolManager->DeactivateTool(EToolSide::Mouse, EToolShutdownType::Cancel);
+			}
+			return true;
 		}
-		return true;
 	}
 
 	// enter key accepts current tool, or ends tool if it does not have accept state
diff --git a/Engine/Source/Editor/EditorInteractiveToolsFramework/Public/EdModeInteractiveToolsContext.h b/Engine/Source/Editor/EditorInteractiveToolsFramework/Public/EdModeInteractiveToolsContext.h
--- a/Engine/Source/Editor/EditorInteractiveToolsFramework/Public/EdModeInteractiveToolsContext.h
+++ b/Engine/Source/Editor/EditorInteractiveToolsFramework/Public/EdModeInteractiveToolsContext.h
@@ -56,6 +56,10 @@ protected:
 	virtual void Shutdown() override;
 
 
+public:
+	UPROPERTY()
+	UMaterialInterface* StandardVertexColorMaterial;
+
 protected:
 	FEdMode* EditorMode;
 
diff --git a/Engine/Source/Runtime/InteractiveToolsFramework/Private/InteractiveToolManager.cpp b/Engine/Source/Runtime/InteractiveToolsFramework/Private/InteractiveToolManager.cpp
--- a/Engine/Source/Runtime/InteractiveToolsFramework/Private/InteractiveToolManager.cpp
+++ b/Engine/Source/Runtime/InteractiveToolsFramework/Private/InteractiveToolManager.cpp
@@ -130,6 +130,8 @@ bool UInteractiveToolManager::ActivateTool(EToolSide Side)
 
 	PostInvalidation();
 
+	OnToolStarted.Broadcast(this, ActiveLeftTool);
+
 	return true;
 }
 
@@ -146,9 +148,12 @@ void UInteractiveToolManager::DeactivateTool(EToolSide Side, EToolShutdownType S
 
 		InputRouter->DeregisterSource(ActiveLeftTool);
 
+		UInteractiveTool* DoneTool = ActiveLeftTool;
 		ActiveLeftTool = nullptr;
 
 		PostInvalidation();
+
+		OnToolEnded.Broadcast(this, DoneTool);
 	}
 }
 
@@ -254,4 +259,7 @@ void UInteractiveToolManager::EmitObjectChange(UObject* TargetObject, TUniquePtr
 	TransactionsAPI->AppendChange(TargetObject, MoveTemp(Change), Description );
 }
 
-
+bool UInteractiveToolManager::RequestSelectionChange(const FSelectedOjectsChangeList& SelectionChange)
+{
+	return TransactionsAPI->RequestSelectionChange(SelectionChange);
+}
diff --git a/Engine/Source/Runtime/InteractiveToolsFramework/Public/InteractiveToolManager.h b/Engine/Source/Runtime/InteractiveToolsFramework/Public/InteractiveToolManager.h
--- a/Engine/Source/Runtime/InteractiveToolsFramework/Public/InteractiveToolManager.h
+++ b/Engine/Source/Runtime/InteractiveToolsFramework/Public/InteractiveToolManager.h
@@ -167,6 +167,13 @@ public:
 	virtual void EmitObjectChange(UObject* TargetObject, TUniquePtr<FChange> Change, const FText& Description );
 
 
+	/**
+	 * Forward an FChange object to the Context
+	 */
+	virtual bool RequestSelectionChange(const FSelectedOjectsChangeList& SelectionChange);
+
+
+
 	//
 	// State control  (@todo: have the Context call these? not safe for anyone to call)
 	//
@@ -178,6 +185,15 @@ public:
 	virtual void Render(IToolsContextRenderAPI* RenderAPI);
 
 
+	//
+	// access to APIs, etc
+	//
+	
+	/** @return current IToolsContextQueriesAPI */
+	virtual IToolsContextQueriesAPI* GetContextQueriesAPI() { return QueriesAPI; }
+
+
+
 public:
 	/** Currently-active Left Tool, or null if no Tool is active */
 	UPROPERTY()
@@ -188,6 +204,16 @@ public:
 	UInteractiveTool* ActiveRightTool;
 
 
+
+
+public:
+	DECLARE_MULTICAST_DELEGATE_TwoParams(FToolManagerToolStartedSignature, UInteractiveToolManager*, UInteractiveTool*);
+	FToolManagerToolStartedSignature OnToolStarted;
+
+	DECLARE_MULTICAST_DELEGATE_TwoParams(FToolManagerToolEndedSignature, UInteractiveToolManager*, UInteractiveTool*);
+	FToolManagerToolEndedSignature OnToolEnded;
+
+
 protected:
 	/** Pointer to current Context-Queries implementation */
 	IToolsContextQueriesAPI* QueriesAPI;
diff --git a/Engine/Source/Runtime/InteractiveToolsFramework/Public/ToolContextInterfaces.h b/Engine/Source/Runtime/InteractiveToolsFramework/Public/ToolContextInterfaces.h
--- a/Engine/Source/Runtime/InteractiveToolsFramework/Public/ToolContextInterfaces.h
+++ b/Engine/Source/Runtime/InteractiveToolsFramework/Public/ToolContextInterfaces.h
@@ -42,6 +42,14 @@ struct INTERACTIVETOOLSFRAMEWORK_API FToolBuilderState
 };
 
 
+/** Types of standard materials that Tools may request from Context */
+UENUM()
+enum class EStandardToolContextMaterials
+{
+	/** White material that displays vertex colors set on mesh */
+	VertexColorMaterial = 1
+};
+
 
 
 /**
@@ -58,6 +66,14 @@ public:
 	 * @param StateOut this structure is populated with available state information
 	 */
 	virtual void GetCurrentSelectionState(FToolBuilderState& StateOut) const = 0;
+
+
+	/**
+	 * Many tools need standard types of materials that the user should provide (eg a vertex-color material, etc)
+	 * @param MaterialType the type of material being requested
+	 * @return Instance of material to use for this purpose
+	 */
+	virtual UMaterialInterface* GetStandardMaterial(EStandardToolContextMaterials MaterialType) const = 0;
 };
 
 
@@ -77,6 +93,28 @@ enum class EToolMessageLevel
 	UserError = 3
 };
 
+/** Type of change we want to apply to a selection */
+UENUM()
+enum class ESelectedObjectsModificationType
+{
+	Replace = 0,
+	Add = 1,
+	Remove = 2,
+	Clear = 3
+};
+
+
+/** Represents a change to a set of selected Actors and Components */
+struct FSelectedOjectsChangeList
+{
+	/** How should this list be interpreted in the context of a larger selection set */
+	ESelectedObjectsModificationType ModificationType;
+	/** List of Actors */
+	TArray<AActor*> Actors;
+	/** List of Componets */
+	TArray<UActorComponent*> Components;
+};
+
 
 /**
  * Users of the Tools Framework need to implement IToolsContextTransactionsAPI so that
@@ -124,6 +162,16 @@ public:
 	 * @param Description text description of the transaction that could be shown to user
 	 */
 	virtual void AppendChange(UObject* TargetObject, TUniquePtr<FChange> Change, const FText& Description) = 0;
+
+
+
+	/**
+	 * Request a modification to the current selected objects
+	 * @param SelectionChange desired modification to current selection
+	 * @return true if the selection change could be applied
+	 */
+	virtual bool RequestSelectionChange(const FSelectedOjectsChangeList& SelectionChange) = 0;
+
 };
 
 
