add missing header
#rnx

diff --git a/Engine/Source/Editor/EditorInteractiveToolsFramework/Private/EdModeInteractiveToolsContext.cpp b/Engine/Source/Editor/EditorInteractiveToolsFramework/Private/EdModeInteractiveToolsContext.cpp
--- a/Engine/Source/Editor/EditorInteractiveToolsFramework/Private/EdModeInteractiveToolsContext.cpp
+++ b/Engine/Source/Editor/EditorInteractiveToolsFramework/Private/EdModeInteractiveToolsContext.cpp
@@ -7,6 +7,7 @@
 #include "Engine/Selection.h"
 #include "Misc/ITransaction.h"
 #include "ScopedTransaction.h"
+#include "Materials/Material.h"
 
 #include "EditorToolAssetAPI.h"
 #include "EditorComponentSourceFactory.h"
