add IToolsContextQueriesAPI::GetCurrentViewState (camera state query), implement in editor tools context
#rnx

diff --git a/Engine/Source/Editor/EditorInteractiveToolsFramework/Private/EdModeInteractiveToolsContext.cpp b/Engine/Source/Editor/EditorInteractiveToolsFramework/Private/EdModeInteractiveToolsContext.cpp
--- a/Engine/Source/Editor/EditorInteractiveToolsFramework/Private/EdModeInteractiveToolsContext.cpp
+++ b/Engine/Source/Editor/EditorInteractiveToolsFramework/Private/EdModeInteractiveToolsContext.cpp
@@ -4,6 +4,7 @@
 #include "EdModeInteractiveToolsContext.h"
 #include "EditorViewportClient.h"
 #include "EditorModeManager.h"
+#include "LevelEditorViewport.h"   // for GCurrentLevelEditingViewportClient
 #include "Engine/Selection.h"
 #include "Misc/ITransaction.h"
 #include "ScopedTransaction.h"
@@ -13,6 +14,7 @@
 #include "EditorComponentSourceFactory.h"
 
 
+
 //#define ENABLE_DEBUG_PRINTING
 
 
@@ -23,6 +25,8 @@ class FEdModeToolsContextQueriesImpl : public IToolsContextQueriesAPI
 public:
 	UEdModeInteractiveToolsContext* ToolsContext;
 	FEdMode* EditorMode;
+	
+	FViewCameraState CachedViewState;
 
 	FEdModeToolsContextQueriesImpl(UEdModeInteractiveToolsContext* Context, FEdMode* EditorModeIn)
 	{
@@ -30,6 +34,15 @@ public:
 		EditorMode = EditorModeIn;
 	}
 
+	void CacheCurrentViewState(FEditorViewportClient* ViewportClient)
+	{
+		FViewportCameraTransform ViewTransform = ViewportClient->GetViewTransform();
+		CachedViewState.Position = ViewTransform.GetLocation();
+		CachedViewState.Orientation = ViewTransform.GetRotation().Quaternion();
+		CachedViewState.bIsOrthographic = ViewportClient->IsOrtho();
+		CachedViewState.bIsVR = false;
+	}
+
 	virtual void GetCurrentSelectionState(FToolBuilderState& StateOut) const override
 	{
 		StateOut.ToolManager = ToolsContext->ToolManager;
@@ -40,6 +53,12 @@ public:
 		StateOut.SourceBuilder = ToolsContext->GetComponentSourceFactory();
 	}
 
+
+	virtual void GetCurrentViewState(FViewCameraState& StateOut) const override
+	{
+		StateOut = CachedViewState;
+	}
+
 	virtual UMaterialInterface* GetStandardMaterial(EStandardToolContextMaterials MaterialType) const
 	{
 		if (MaterialType == EStandardToolContextMaterials::VertexColorMaterial)
@@ -226,6 +245,14 @@ void UEdModeInteractiveToolsContext::Tick(FEditorViewportClient* ViewportClient,
 		ViewportClient->Invalidate();
 		bInvalidationPending = false;
 	}
+
+	// save this view
+	// Check against GCurrentLevelEditingViewportClient is temporary and should be removed in future.
+	// Current issue is that this ::Tick() is called *per viewport*, so once for each view in a 4-up view.
+	if (ViewportClient == GCurrentLevelEditingViewportClient)
+	{
+		((FEdModeToolsContextQueriesImpl*)this->QueriesAPI)->CacheCurrentViewState(ViewportClient);
+	}
 }
 
 
diff --git a/Engine/Source/Runtime/InteractiveToolsFramework/Public/ToolContextInterfaces.h b/Engine/Source/Runtime/InteractiveToolsFramework/Public/ToolContextInterfaces.h
--- a/Engine/Source/Runtime/InteractiveToolsFramework/Public/ToolContextInterfaces.h
+++ b/Engine/Source/Runtime/InteractiveToolsFramework/Public/ToolContextInterfaces.h
@@ -42,6 +42,23 @@ struct INTERACTIVETOOLSFRAMEWORK_API FToolBuilderState
 };
 
 
+/**
+ * FViewCameraState is a bucket of state information that a Tool might
+ * need to implement interactions that depend on the current scene view.
+ */
+struct INTERACTIVETOOLSFRAMEWORK_API FViewCameraState
+{
+	/** Current camera/head position */
+	FVector Position;
+	/** Current camera/head orientation */
+	FQuat Orientation;
+	/** Is current view an orthographic view */
+	bool bIsOrthographic;
+	/** Is current view a VR view */
+	bool bIsVR;
+};
+
+
 /** Types of standard materials that Tools may request from Context */
 UENUM()
 enum class EStandardToolContextMaterials
@@ -68,6 +85,13 @@ public:
 	virtual void GetCurrentSelectionState(FToolBuilderState& StateOut) const = 0;
 
 
+	/**
+	 * Request information about current view state
+	 * @param StateOut this structure is populated with available state information
+	 */
+	virtual void GetCurrentViewState(FViewCameraState& StateOut) const = 0;
+
+
 	/**
 	 * Many tools need standard types of materials that the user should provide (eg a vertex-color material, etc)
 	 * @param MaterialType the type of material being requested
