Add basic support for running vertex and edge snap queries against the current world. Only works with objects that can be hit by linetrace
#rnx

diff --git a/Engine/Source/Editor/EditorInteractiveToolsFramework/Private/EdModeInteractiveToolsContext.cpp b/Engine/Source/Editor/EditorInteractiveToolsFramework/Private/EdModeInteractiveToolsContext.cpp
--- a/Engine/Source/Editor/EditorInteractiveToolsFramework/Private/EdModeInteractiveToolsContext.cpp
+++ b/Engine/Source/Editor/EditorInteractiveToolsFramework/Private/EdModeInteractiveToolsContext.cpp
@@ -9,17 +9,19 @@
 #include "Misc/ITransaction.h"
 #include "ScopedTransaction.h"
 #include "Materials/Material.h"
+#include "Engine/StaticMesh.h"
+#include "Components/StaticMeshComponent.h"
 
 #include "EditorToolAssetAPI.h"
 #include "EditorComponentSourceFactory.h"
 
-
+//#include "PhysicsEngine/BodySetup.h"
+//#include "Interfaces/Interface_CollisionDataProvider.h"
 
 //#define ENABLE_DEBUG_PRINTING
 
 
 
-
 class FEdModeToolsContextQueriesImpl : public IToolsContextQueriesAPI
 {
 public:
@@ -59,6 +61,155 @@ public:
 		StateOut = CachedViewState;
 	}
 
+
+	virtual bool ExecuteSceneSnapQuery(const FSceneSnapQueryRequest& Request, TArray<FSceneSnapQueryResult>& Results) const override
+	{
+		if (Request.RequestType != ESceneSnapQueryType::Position)
+		{
+			return false;		// not supported yet
+		}
+
+		int FoundResultCount = 0;
+
+		//
+		// Run a snap query by casting ray into the world.
+		// If a hit is found, we look up what triangle was hit, and then test its vertices and edges
+		// 
+
+		// cast ray into world
+		FVector RayStart = CachedViewState.Position;
+		FVector RayDirection = Request.Position - RayStart; RayDirection.Normalize();
+		FVector RayEnd = RayStart + 9999999 * RayDirection;
+		FCollisionObjectQueryParams ObjectQueryParams(FCollisionObjectQueryParams::AllObjects);
+		FCollisionQueryParams QueryParams = FCollisionQueryParams::DefaultQueryParam;
+		QueryParams.bTraceComplex = true;
+		QueryParams.bReturnFaceIndex = true;
+		FHitResult HitResult;
+		bool bHitWorld = EditorMode->GetWorld()->LineTraceSingleByObjectType(HitResult, RayStart, RayEnd, ObjectQueryParams, QueryParams);
+		if (bHitWorld && HitResult.FaceIndex >= 0) 
+		{
+			float VisualAngle = OpeningAngleDeg(Request.Position, HitResult.ImpactPoint, RayStart);
+			//UE_LOG(LogTemp, Warning, TEXT("[HIT] visualangle %f faceindex %d"), VisualAngle, HitResult.FaceIndex);
+			if (VisualAngle < Request.VisualAngleThresholdDegrees)
+			{
+				UPrimitiveComponent* Component = HitResult.Component.Get();
+				if (Cast<UStaticMeshComponent>(Component) != nullptr)
+				{
+					// HitResult.FaceIndex is apparently an index into the TriMeshCollisionData, not sure how
+					// to directly access it. Calling GetPhysicsTriMeshData is expensive!
+					//UBodySetup* BodySetup = Cast<UStaticMeshComponent>(Component)->GetBodySetup();
+					//UObject* CDPObj = BodySetup->GetOuter();
+					//IInterface_CollisionDataProvider* CDP = Cast<IInterface_CollisionDataProvider>(CDPObj);
+					//FTriMeshCollisionData TriMesh;
+					//CDP->GetPhysicsTriMeshData(&TriMesh, true);
+					//FTriIndices Triangle = TriMesh.Indices[HitResult.FaceIndex];
+					//FVector Positions[3] = { TriMesh.Vertices[Triangle.v0], TriMesh.Vertices[Triangle.v1], TriMesh.Vertices[Triangle.v2] };
+
+					// physics collision data is created from StaticMesh RenderData
+					// so use HitResult.FaceIndex to extract triangle from the LOD0 mesh
+					// (note: this may be incorrect if there are multiple sections...in that case I think we have to 
+					//  first find section whose accumulated index range would contain .FaceIndexX)
+					UStaticMesh* StaticMesh = Cast<UStaticMeshComponent>(Component)->GetStaticMesh();
+					FStaticMeshLODResources& LOD = StaticMesh->RenderData->LODResources[0];
+					FIndexArrayView Indices = LOD.IndexBuffer.GetArrayView();
+					int TriIdx = 3 * HitResult.FaceIndex;
+					int Triangle[3] = { Indices[TriIdx], Indices[TriIdx + 1], Indices[TriIdx + 2] };
+					FVector Positions[3] = {
+						LOD.VertexBuffers.PositionVertexBuffer.VertexPosition(Triangle[0]),
+						LOD.VertexBuffers.PositionVertexBuffer.VertexPosition(Triangle[1]),
+						LOD.VertexBuffers.PositionVertexBuffer.VertexPosition(Triangle[2])
+					};
+					
+					// transform to world space
+					FTransform ComponentTransform = Component->GetComponentTransform();
+					Positions[0] = ComponentTransform.TransformPosition(Positions[0]);
+					Positions[1] = ComponentTransform.TransformPosition(Positions[1]);
+					Positions[2] = ComponentTransform.TransformPosition(Positions[2]);
+
+					FSceneSnapQueryResult SnapResult;
+					SnapResult.TriVertices[0] = Positions[0];
+					SnapResult.TriVertices[1] = Positions[1];
+					SnapResult.TriVertices[2] = Positions[2];
+
+					// try snapping to vertices
+					float SmallestAngle = Request.VisualAngleThresholdDegrees;
+					if ( (Request.TargetTypes & ESceneSnapQueryTargetType::MeshVertex) != ESceneSnapQueryTargetType::None)
+					{
+						for (int j = 0; j < 3; ++j)
+						{
+							VisualAngle = OpeningAngleDeg(Request.Position, Positions[j], RayStart);
+							if (VisualAngle < SmallestAngle)
+							{
+								SmallestAngle = VisualAngle;
+								SnapResult.Position = Positions[j];
+								SnapResult.TargetType = ESceneSnapQueryTargetType::MeshVertex;
+							}
+						}
+					}
+
+					// try snapping to nearest points on edges
+					if ( ((Request.TargetTypes & ESceneSnapQueryTargetType::MeshEdge) != ESceneSnapQueryTargetType::None) &&
+						 (SnapResult.TargetType != ESceneSnapQueryTargetType::MeshVertex) )
+					{
+						for (int j = 0; j < 3; ++j)
+						{
+							FVector EdgeNearestPt = NearestSegmentPt(Positions[j], Positions[(j+1)%3], Request.Position);
+							VisualAngle = OpeningAngleDeg(Request.Position, EdgeNearestPt, RayStart);
+							if (VisualAngle < SmallestAngle )
+							{
+								SmallestAngle = VisualAngle;
+								SnapResult.Position = EdgeNearestPt;
+								SnapResult.TargetType = ESceneSnapQueryTargetType::MeshEdge;
+							}
+						}
+					}
+
+					// if we found a valid snap, return it
+					if (SmallestAngle < Request.VisualAngleThresholdDegrees)
+					{
+						SnapResult.TargetActor = HitResult.Actor.Get();
+						SnapResult.TargetComponent = HitResult.Component.Get();
+						Results.Add(SnapResult);
+						FoundResultCount++;
+					}
+				}
+			}
+
+		}
+
+		return (FoundResultCount > 0);
+	}
+
+
+	//@ todo this are mirrored from GeometryProcessing, which is still experimental...replace w/ direct calls once GP component is standardized
+	static float OpeningAngleDeg(FVector A, FVector B, const FVector& P)
+	{
+		A -= P;
+		A.Normalize();
+		B -= P;
+		B.Normalize();
+		float Dot = FMath::Clamp(FVector::DotProduct(A,B), -1.0f, 1.0f);
+		return acos(Dot) * (180.0f / 3.141592653589f);
+	}
+	static FVector NearestSegmentPt(FVector A, FVector B, const FVector& P)
+	{
+		FVector Direction = (B - A); 
+		float Length = Direction.Size();
+		Direction /= Length;
+		float t = FVector::DotProduct( (P - A), Direction);
+		if (t >= Length)
+		{
+			return B;
+		}
+		if (t <= 0)
+		{
+			return A;
+		}
+		return A + t * Direction;
+	}
+
+
+
 	virtual UMaterialInterface* GetStandardMaterial(EStandardToolContextMaterials MaterialType) const
 	{
 		if (MaterialType == EStandardToolContextMaterials::VertexColorMaterial)
diff --git a/Engine/Source/Runtime/InteractiveToolsFramework/Public/ToolContextInterfaces.h b/Engine/Source/Runtime/InteractiveToolsFramework/Public/ToolContextInterfaces.h
--- a/Engine/Source/Runtime/InteractiveToolsFramework/Public/ToolContextInterfaces.h
+++ b/Engine/Source/Runtime/InteractiveToolsFramework/Public/ToolContextInterfaces.h
@@ -59,6 +59,76 @@ struct INTERACTIVETOOLSFRAMEWORK_API FViewCameraState
 };
 
 
+
+/** Types of Snap Queries that a ToolsContext parent may support, that Tools may request */
+UENUM()
+enum class ESceneSnapQueryType
+{
+	/** snapping a position */
+	Position = 1		
+};
+
+/** Types of snap targets that a Tool may want to run snap queries against. */
+UENUM()
+enum class ESceneSnapQueryTargetType
+{
+	None = 0,
+	/** Consider any mesh vertex */
+	MeshVertex = 1,
+	/** Consider any mesh edge */
+	MeshEdge = 2,
+
+	All = MeshVertex | MeshEdge
+};
+ENUM_CLASS_FLAGS(ESceneSnapQueryTargetType);
+
+/**
+ * Configuration variables for a IToolsContextQueriesAPI snap query request.
+ */
+struct INTERACTIVETOOLSFRAMEWORK_API FSceneSnapQueryRequest
+{
+	/** What type of snap query geometry is this */
+	ESceneSnapQueryType RequestType;
+	/** What does caller want to try to snap to */
+	ESceneSnapQueryTargetType TargetTypes;
+
+	/** Snap input position */
+	FVector Position;
+	/** Another position must deviate less than this number of degrees (in visual angle) to be considered an acceptable snap position */
+	float VisualAngleThresholdDegrees;
+
+	/** Snap input direction */
+	FVector Direction;
+	/** Another direction must deviate less than this number of degrees from Direction to be considered an acceptable snap direction */
+	float DirectionAngleThresholdDegrees;
+};
+
+
+/**
+ * Computed result of a IToolsContextQueriesAPI snap query request
+ */
+struct INTERACTIVETOOLSFRAMEWORK_API FSceneSnapQueryResult
+{
+	/** Actor that owns snap target */
+	AActor* TargetActor = nullptr;
+	/** Component that owns snap target */
+	UActorComponent* TargetComponent = nullptr;
+	/** What kind of geometric element was snapped to */
+	ESceneSnapQueryTargetType TargetType = ESceneSnapQueryTargetType::None;
+
+	/** Snap position (may not be set depending on query types) */
+	FVector Position;
+	/** Snap normal (may not be set depending on query types) */
+	FVector Normal;
+	/** Snap direction (may not be set depending on query types) */
+	FVector Direction;
+
+	/** Vertices of triangle that contains result (for debugging, may not be set) */
+	FVector TriVertices[3];
+};
+
+
+
 /** Types of standard materials that Tools may request from Context */
 UENUM()
 enum class EStandardToolContextMaterials
@@ -92,6 +162,16 @@ public:
 	virtual void GetCurrentViewState(FViewCameraState& StateOut) const = 0;
 
 
+	/**
+	 * Try to find Snap Targets in the scene that satisfy the Snap Query.
+	 * @param Request snap query configuration
+	 * @param Results list of potential snap results
+	 * @return true if any valid snap target was found
+	 * @warning implementations are not required (and may not be able) to support snapping
+	 */
+	virtual bool ExecuteSceneSnapQuery(const FSceneSnapQueryRequest& Request, TArray<FSceneSnapQueryResult>& Results ) const = 0;
+
+
 	/**
 	 * Many tools need standard types of materials that the user should provide (eg a vertex-color material, etc)
 	 * @param MaterialType the type of material being requested
