fix linux and clang builds
#rnx

diff --git a/Engine/Source/Editor/EditorInteractiveToolsFramework/Private/EdModeInteractiveToolsContext.cpp b/Engine/Source/Editor/EditorInteractiveToolsFramework/Private/EdModeInteractiveToolsContext.cpp
--- a/Engine/Source/Editor/EditorInteractiveToolsFramework/Private/EdModeInteractiveToolsContext.cpp
+++ b/Engine/Source/Editor/EditorInteractiveToolsFramework/Private/EdModeInteractiveToolsContext.cpp
@@ -112,13 +112,11 @@ public:
 					UStaticMesh* StaticMesh = Cast<UStaticMeshComponent>(Component)->GetStaticMesh();
 					FStaticMeshLODResources& LOD = StaticMesh->RenderData->LODResources[0];
 					FIndexArrayView Indices = LOD.IndexBuffer.GetArrayView();
-					int TriIdx = 3 * HitResult.FaceIndex;
-					int Triangle[3] = { Indices[TriIdx], Indices[TriIdx + 1], Indices[TriIdx + 2] };
-					FVector Positions[3] = {
-						LOD.VertexBuffers.PositionVertexBuffer.VertexPosition(Triangle[0]),
-						LOD.VertexBuffers.PositionVertexBuffer.VertexPosition(Triangle[1]),
-						LOD.VertexBuffers.PositionVertexBuffer.VertexPosition(Triangle[2])
-					};
+					int32 TriIdx = 3 * HitResult.FaceIndex;
+					FVector Positions[3];
+					Positions[0] = LOD.VertexBuffers.PositionVertexBuffer.VertexPosition(Indices[TriIdx]);
+					Positions[1] = LOD.VertexBuffers.PositionVertexBuffer.VertexPosition(Indices[TriIdx+1]);
+					Positions[2] = LOD.VertexBuffers.PositionVertexBuffer.VertexPosition(Indices[TriIdx+2]);
 					
 					// transform to world space
 					FTransform ComponentTransform = Component->GetComponentTransform();
