Add FStandardToolModeCommands for tool hotkeys shareable across FEdModes, and TInteractiveToolCommands which connects editor hotkey system to ToolsFramework

#rnx

diff --git a/Engine/Source/Editor/EditorInteractiveToolsFramework/EditorInteractiveToolsFramework.Build.cs b/Engine/Source/Editor/EditorInteractiveToolsFramework/EditorInteractiveToolsFramework.Build.cs
--- a/Engine/Source/Editor/EditorInteractiveToolsFramework/EditorInteractiveToolsFramework.Build.cs
+++ b/Engine/Source/Editor/EditorInteractiveToolsFramework/EditorInteractiveToolsFramework.Build.cs
@@ -37,10 +37,14 @@ public class EditorInteractiveToolsFramework : ModuleRules
 			{
 				"CoreUObject",
 				"Engine",
-				"InputCore",
+                "Slate",
+                "SlateCore",
+                "InputCore",
 				"UnrealEd",
                 "ContentBrowser",
                 "LevelEditor",
+                "ApplicationCore",
+                "EditorStyle",
                 "InteractiveToolsFramework"
 
 				// ... add private dependencies that you statically link with here ...	
diff --git a/Engine/Source/Editor/EditorInteractiveToolsFramework/Private/StandardToolModeCommands.cpp b/Engine/Source/Editor/EditorInteractiveToolsFramework/Private/StandardToolModeCommands.cpp
new file mode 100644
--- /dev/null
+++ b/Engine/Source/Editor/EditorInteractiveToolsFramework/Private/StandardToolModeCommands.cpp
@@ -0,0 +1,38 @@
+// Copyright 1998-2019 Epic Games, Inc. All Rights Reserved.
+
+#include "StandardToolModeCommands.h"
+#include "EditorStyleSet.h"
+
+#define LOCTEXT_NAMESPACE "StandardToolModeCommands"
+
+
+FStandardToolModeCommands::FStandardToolModeCommands() :
+	TCommands<FStandardToolModeCommands>(
+		"StandardToolCommands", // Context name for fast lookup
+		NSLOCTEXT("Contexts", "StandardToolCommands", "Standard Tool Commands"), // Localized context name for displaying
+		NAME_None, // Parent
+		FEditorStyle::GetStyleSetName() // Icon Style Set
+		)
+{}
+
+
+void FStandardToolModeCommands::RegisterCommands()
+{
+	TSharedPtr<FUICommandInfo> IncreaseBrushSize;
+	UI_COMMAND(IncreaseBrushSize, "Increase Brush Size", "Increases the size of the brush", EUserInterfaceActionType::Button, FInputChord(EKeys::RightBracket));
+	Commands.Add(EStandardToolModeCommands::IncreaseBrushSize, IncreaseBrushSize);
+
+	TSharedPtr<FUICommandInfo> DecreaseBrushSize;
+	UI_COMMAND(DecreaseBrushSize, "Decrease Brush Size", "Decreases the size of the brush", EUserInterfaceActionType::Button, FInputChord(EKeys::LeftBracket));
+	Commands.Add(EStandardToolModeCommands::DecreaseBrushSize, DecreaseBrushSize);
+}
+
+
+TSharedPtr<FUICommandInfo> FStandardToolModeCommands::FindStandardCommand(EStandardToolModeCommands Command) const
+{
+	const TSharedPtr<FUICommandInfo>* Found = Commands.Find(Command);
+	checkf(Found != nullptr, TEXT("FStandardToolModeCommands::FindStandardCommand: standard command %d was not found!"), Command);
+	return (Found != nullptr) ? *Found : nullptr;
+}
+
+
diff --git a/Engine/Source/Editor/EditorInteractiveToolsFramework/Public/InteractiveToolsCommands.h b/Engine/Source/Editor/EditorInteractiveToolsFramework/Public/InteractiveToolsCommands.h
new file mode 100644
--- /dev/null
+++ b/Engine/Source/Editor/EditorInteractiveToolsFramework/Public/InteractiveToolsCommands.h
@@ -0,0 +1,302 @@
+// Copyright 1998-2019 Epic Games, Inc. All Rights Reserved.
+
+#pragma once
+
+#include "CoreMinimal.h"
+#include "Framework/Commands/Commands.h"
+#include "Framework/Commands/UICommandList.h"
+#include "InteractiveTool.h"
+#include "StandardToolModeCommands.h"
+
+
+
+/**
+ * TInteractiveToolCommands is a base class that handles connecting up Tool Actions 
+ * (ie the FInteractiveToolAction provided by a UInteractiveTool) to the UnrealEditor
+ * Command system, which allows for remappable hotkeys, etc
+ * 
+ * Usage is as follows:
+ *    - in your EdMode Module, create a subclass-instance of this, say FMyToolCommands and call FMyToolCommands::Register() in your ::StartupModule() function
+ *        - subclass must implement ::GetToolDefaultObjectList(), here you just add GetMutableDefault<MyToolX> to the input list for all your Tools
+ *        - subclass must implement ::GetNamespaceString(), likely you can just return TEXT(LOCTEXT_NAMESPACE)
+ *    - add a member TSharedPtr<FUICommandList> UICommandList; to your EdMode impl
+ *    - when you start a new Tool, call FMyToolCommands::Get().BindCommandsForCurrentTool(UICommandList, NewTool)
+ *    - when you end a Tool, call FMyToolCommands::Get().UnbindActiveCommands()
+ *    - in your EdModeImpl::InputKey override, call UICommandList->ProcessCommandBindings()
+ */
+template<typename CommandContextType>
+class TInteractiveToolCommands : public TCommands<CommandContextType>
+{
+
+public:
+	/**
+	 * Initialize commands. Collects possible Actions from the available UInteractiveTools
+	 * (provided by GetToolDefaultObjectList) and then registers a FUICommandInfo for each.
+	 * This needs to be called in ::StartupModule for your EdMode Module
+	 */
+	virtual void RegisterCommands() override;
+
+	/**
+	 * Bind any of the registered UICommands to the given Tool, if they are compatible.
+	 */
+	virtual void BindCommandsForCurrentTool(TSharedPtr<FUICommandList> UICommandList, UInteractiveTool* Tool) const;
+
+	/**
+	 * Unbind all of the currently-registered commands
+	 */
+	virtual void UnbindActiveCommands(TSharedPtr<FUICommandList> UICommandList) const;
+
+
+	//
+	// Interface that subclasses need to implement
+	// 
+
+
+	/**
+	 * RegisterCommands() needs actual UInteractiveTool instances for all the Tools that want to
+	 * provide Actions which will be connected to hotkeys. We can do this based on the CDO's for
+	 * each tool, returned by GetMutableDefault<UYourInteractiveToolType>(). The Tool CDOs are
+	 * not owned by a ToolManager and we will only call .GetActionSet() on them.
+	 */
+	virtual void GetToolDefaultObjectList(TArray<UInteractiveTool*>& ToolCDOs) = 0;
+
+	/**
+	 * The Namespace used to identify this Command set
+	 */
+	virtual const TCHAR * GetNamespaceString() = 0;
+
+
+protected:
+
+	/** Forwarding constructor */
+	TInteractiveToolCommands(const FName InContextName, const FText& InContextDesc, const FName InContextParent, const FName InStyleSetName)
+		: TCommands(InContextName, InContextDesc, InContextParent, InStyleSetName)
+	{
+	}
+
+	/**
+	 * Query FStandardToolModeCommands to find an existing command/hotkey for this standard tool action
+	 */
+	virtual TSharedPtr<FUICommandInfo> FindStandardCommand(EStandardToolActions ToolActionID);
+
+	/**
+	 * List of pairs of known Tool Actions and their associated UICommands.
+	 * This list is used at bind-time to connect up the already-registered UICommands to the active Tool
+	 */
+	struct FToolActionCommand
+	{
+		FInteractiveToolAction ToolAction;
+		TSharedPtr<FUICommandInfo> UICommand;
+	};
+	TArray<FToolActionCommand> ActionCommands;
+
+
+	//
+	// Support for sharing common commands between Tools where the commands
+	// are not shared across Tool Modes (and hence not in FStandardToolModeCommands)
+	// 
+
+	/** List FToolActionCommands for standard Tool Actions */
+	TMap<EStandardToolActions, FToolActionCommand> SharedActionCommands;
+
+	/**
+	 * Find or Create a UICommand for a standard Tool Action, that will be shared across Tools
+	 */
+	bool FindOrCreateSharedCommand(EStandardToolActions ActionID, FToolActionCommand& FoundOut);
+
+	/**
+	 * Create a suitable FInteractiveToolAction for the standard Tool Action, ie with suitable
+	 * command name and description strings and hotkeys. This info will be used to create the shared UICommand.
+	 */
+	virtual bool IntializeStandardToolAction(EStandardToolActions ActionID, FInteractiveToolAction& ToolActionOut);
+
+	/**
+	 * Utility function that registeres a Tool Aciton as a UICommand
+	 */
+	void RegisterUIToolCommand(const FInteractiveToolAction& ToolAction, TSharedPtr<FUICommandInfo>& UICommandInfo);
+};
+
+
+
+
+
+
+template<typename CommandContextType>
+void TInteractiveToolCommands<CommandContextType>::RegisterCommands()
+{
+	// make sure standard commands are registered
+	FStandardToolModeCommands::Register();
+
+	// get list of all tools used by command set
+	TArray<UInteractiveTool*> AllToolsCDOs;
+	GetToolDefaultObjectList(AllToolsCDOs);
+
+	// collect all the actions in these tools
+	TArray<FInteractiveToolAction> ToolActions;
+	for (UInteractiveTool* Tool : AllToolsCDOs)
+	{
+		Tool->GetActionSet()->CollectActions(ToolActions);
+	}
+
+	// register all the commands in all the actions
+	int NumActions = ToolActions.Num();
+	ActionCommands.SetNum(NumActions);
+	for (int k = 0; k < NumActions; ++k)
+	{
+		const FInteractiveToolAction& ToolAction = ToolActions[k];
+		ActionCommands[k].ToolAction = ToolAction;
+
+		bool bRegistered = false;
+		if ((int32)ToolAction.ActionID < (int32)EStandardToolActions::BaseClientDefinedActionID)
+		{
+			TSharedPtr<FUICommandInfo> FoundStandard = FindStandardCommand((EStandardToolActions)ToolAction.ActionID);
+			if (FoundStandard.IsValid())
+			{
+				ActionCommands[k].UICommand = FoundStandard;
+				bRegistered = true;
+			}
+			else
+			{
+				FToolActionCommand StandardActionCommand;
+				bool bFoundRemap = FindOrCreateSharedCommand((EStandardToolActions)ToolAction.ActionID, StandardActionCommand);
+				if (bFoundRemap)
+				{
+					ActionCommands[k].UICommand = StandardActionCommand.UICommand;
+					bRegistered = true;
+				}
+			}
+		}
+
+		if (bRegistered == false)
+		{
+			RegisterUIToolCommand(ActionCommands[k].ToolAction, ActionCommands[k].UICommand);
+		}
+	}
+}
+
+
+template<typename CommandContextType>
+void TInteractiveToolCommands<CommandContextType>::BindCommandsForCurrentTool(TSharedPtr<FUICommandList> UICommandList, UInteractiveTool* Tool) const
+{
+	UClass* ToolClassType = Tool->GetClass();
+
+	int NumActionCommands = ActionCommands.Num();
+	for (int32 k = 0; k < NumActionCommands; ++k)
+	{
+		const FToolActionCommand& ActionCommand = ActionCommands[k];
+		if (ActionCommand.ToolAction.ClassType == ToolClassType)
+		{
+			UICommandList->MapAction(
+				ActionCommand.UICommand,
+				FExecuteAction::CreateUObject(Tool, &UInteractiveTool::ExecuteAction, ActionCommand.ToolAction.ActionID));
+			//FCanExecuteAction::CreateRaw(this, &FEdModeFoliage::CurrentToolUsesBrush));
+		}
+	}
+}
+
+
+template<typename CommandContextType>
+void TInteractiveToolCommands<CommandContextType>::UnbindActiveCommands(TSharedPtr<FUICommandList> UICommandList) const
+{
+	// @todo would be more efficient if we kept track of which commands were mapped.
+	// However currently this function must be const because TCommands::Get() returns a const
+	for (const FToolActionCommand& ActionCommand : ActionCommands)
+	{
+		if (UICommandList->IsActionMapped(ActionCommand.UICommand))
+		{
+			UICommandList->UnmapAction(ActionCommand.UICommand);
+		}
+	}
+}
+
+
+
+
+
+
+//
+// Support for re-using known commands from FStandardToolModeCommands
+//
+
+template<typename CommandContextType>
+TSharedPtr<FUICommandInfo> TInteractiveToolCommands<CommandContextType>::FindStandardCommand(EStandardToolActions ToolActionID)
+{
+	const FStandardToolModeCommands& StdCommands = FStandardToolModeCommands::Get();
+	switch (ToolActionID)
+	{
+	case EStandardToolActions::IncreaseBrushSize:
+		return StdCommands.FindStandardCommand(EStandardToolModeCommands::IncreaseBrushSize);
+	case EStandardToolActions::DecreaseBrushSize:
+		return StdCommands.FindStandardCommand(EStandardToolModeCommands::DecreaseBrushSize);
+	default:
+		return nullptr;
+	}
+}
+
+
+
+
+
+
+//
+// Support for sharing UI commands between multiple Tools, where the
+// shared command is not at the shared-between-modes level
+// 
+
+template<typename CommandContextType>
+bool TInteractiveToolCommands<CommandContextType>::FindOrCreateSharedCommand(EStandardToolActions ActionID, FToolActionCommand& FoundOut)
+{
+	if (SharedActionCommands.Contains(ActionID))
+	{
+		FoundOut = SharedActionCommands[ActionID];
+		return true;
+	}
+
+	FInteractiveToolAction NewAction;
+	bool bIsKnownAction = IntializeStandardToolAction(ActionID, NewAction);
+
+	if (bIsKnownAction)
+	{
+		check(NewAction.ActionID == (int32)ActionID);
+		FToolActionCommand NewActionCommand;
+		NewActionCommand.ToolAction = NewAction;
+		RegisterUIToolCommand(NewActionCommand.ToolAction, NewActionCommand.UICommand);
+		FoundOut = SharedActionCommands.Add(ActionID, NewActionCommand);
+		return true;
+	}
+	return false;
+}
+
+
+template<typename CommandContextType>
+bool TInteractiveToolCommands<CommandContextType>::IntializeStandardToolAction(EStandardToolActions ActionID, FInteractiveToolAction& ToolActionOut)
+{
+	// base class does not have any standard Tool actions. Subclasses can override
+	// this function, populating ToolActionOut and returning true, if the ActionID should be shared between multiple Tools
+	
+	return false;
+}
+
+
+
+
+template<typename CommandContextType>
+void TInteractiveToolCommands<CommandContextType>::RegisterUIToolCommand(const FInteractiveToolAction& ToolAction, TSharedPtr<FUICommandInfo>& UICommandInfo)
+{
+	FString TooltipString = ToolAction.ActionName + FString("_ToolTip");
+	FString DotString = FString(".") + ToolAction.ActionName;
+
+	UI_COMMAND_Function(this,
+		UICommandInfo,
+		GetNamespaceString(),
+		*ToolAction.ActionName,
+		*TooltipString,
+		TCHAR_TO_ANSI(*DotString),
+		*ToolAction.ShortName, *ToolAction.Description,
+		EUserInterfaceActionType::Button,
+		FInputChord(ToolAction.DefaultModifiers, ToolAction.DefaultKey));
+}
+
+
+
+
diff --git a/Engine/Source/Editor/EditorInteractiveToolsFramework/Public/StandardToolModeCommands.h b/Engine/Source/Editor/EditorInteractiveToolsFramework/Public/StandardToolModeCommands.h
new file mode 100644
--- /dev/null
+++ b/Engine/Source/Editor/EditorInteractiveToolsFramework/Public/StandardToolModeCommands.h
@@ -0,0 +1,54 @@
+// Copyright 1998-2019 Epic Games, Inc. All Rights Reserved.
+
+#pragma once
+
+#include "CoreMinimal.h"
+#include "Framework/Commands/Commands.h"
+#include "Framework/Commands/UICommandInfo.h"
+
+
+/**
+ * Integer identifiers for "Standard" Tool commands, that many EditorModes / AssetEditors / etc may share.
+ * This allows a single hotkey binding to be used across many contexts.
+ */
+enum class EStandardToolModeCommands
+{
+	IncreaseBrushSize = 100,
+	DecreaseBrushSize = 101
+};
+
+
+
+/**
+ * FStandardToolModeCommands provides standard commands for Tools. This allows
+ * for "sharing" common hotkey bindings between multiple EditorModes
+ * 
+ * The set of standard commands is defined by EStandardToolModeCommands.
+ * You must call FStandardToolModeCommands::Register() on module startup to register these commands
+ * (note that they may already be registered by another module, but you should call this to be safe).
+ * 
+ * Then, when you want to Bind a standard command in your FUICommandList,
+ * call FStandardToolModeCommands::Get().FindStandardCommand() to get the registered UICommandInfo.
+ * 
+ */
+class EDITORINTERACTIVETOOLSFRAMEWORK_API FStandardToolModeCommands : public TCommands<FStandardToolModeCommands>
+{
+public:
+	FStandardToolModeCommands();
+
+	/**
+	 * Registers the set of standard commands. Call on module startup.
+	 */
+	virtual void RegisterCommands() override;
+
+	/**
+	 * Look up the UICommandInfo for a standard command
+	 */
+	TSharedPtr<FUICommandInfo> FindStandardCommand(EStandardToolModeCommands Command) const;
+
+protected:
+	TMap<EStandardToolModeCommands, TSharedPtr<FUICommandInfo>> Commands;
+
+};
+
+
