add minimal support for keyboard InputBehaviors and keypress capture  (not intended for hotkey actions, but for hold-key modifier usage)
#rnx

diff --git a/Engine/Source/Editor/EditorInteractiveToolsFramework/Private/EdModeInteractiveToolsContext.cpp b/Engine/Source/Editor/EditorInteractiveToolsFramework/Private/EdModeInteractiveToolsContext.cpp
--- a/Engine/Source/Editor/EditorInteractiveToolsFramework/Private/EdModeInteractiveToolsContext.cpp
+++ b/Engine/Source/Editor/EditorInteractiveToolsFramework/Private/EdModeInteractiveToolsContext.cpp
@@ -141,6 +141,7 @@ public:
 								SmallestAngle = VisualAngle;
 								SnapResult.Position = Positions[j];
 								SnapResult.TargetType = ESceneSnapQueryTargetType::MeshVertex;
+								SnapResult.TriSnapIndex = j;
 							}
 						}
 					}
@@ -158,6 +159,7 @@ public:
 								SmallestAngle = VisualAngle;
 								SnapResult.Position = EdgeNearestPt;
 								SnapResult.TargetType = ESceneSnapQueryTargetType::MeshEdge;
+								SnapResult.TriSnapIndex = j;
 							}
 						}
 					}
@@ -480,58 +482,85 @@ bool UEdModeInteractiveToolsContext::InputKey(FEditorViewportClient* ViewportCli
 
 	if (Event == IE_Pressed || Event == IE_Released)
 	{
-		bool bIsLeftMouse = (Key == EKeys::LeftMouseButton);
-		bool bIsMiddleMouse = (Key == EKeys::MiddleMouseButton);
-		bool bIsRightMouse = (Key == EKeys::RightMouseButton);
-
-		if (bIsLeftMouse || bIsMiddleMouse || bIsRightMouse)
+		if (Key.IsMouseButton())
 		{
+			bool bIsLeftMouse = (Key == EKeys::LeftMouseButton);
+			bool bIsMiddleMouse = (Key == EKeys::MiddleMouseButton);
+			bool bIsRightMouse = (Key == EKeys::RightMouseButton);
 
-			// early-out here if we are going to do camera manipulation
-			if (ViewportClient->IsAltPressed())
+			if (bIsLeftMouse || bIsMiddleMouse || bIsRightMouse)
 			{
-				return bHandled;
-			}
 
-			FInputDeviceState InputState = CurrentInputState;
-			InputState.InputDevice = EInputDevices::Mouse;
-			InputState.SetKeyStates(
-				ViewportClient->IsShiftPressed(), ViewportClient->IsAltPressed(),
-				ViewportClient->IsCtrlPressed(), ViewportClient->IsCmdPressed());
+				// early-out here if we are going to do camera manipulation
+				if (ViewportClient->IsAltPressed())
+				{
+					return bHandled;
+				}
 
-			if (bIsLeftMouse)
-			{
-				InputState.Mouse.Left.SetStates(
-					(Event == IE_Pressed), (Event == IE_Pressed), (Event == IE_Released));
-				CurrentInputState.Mouse.Left.bDown = (Event == IE_Pressed);
-			}
-			else if (bIsMiddleMouse)
-			{
-				InputState.Mouse.Middle.SetStates(
-					(Event == IE_Pressed), (Event == IE_Pressed), (Event == IE_Released));
-				CurrentInputState.Mouse.Middle.bDown = (Event == IE_Pressed);
-			}
-			else
-			{
-				InputState.Mouse.Right.SetStates(
-					(Event == IE_Pressed), (Event == IE_Pressed), (Event == IE_Released));
-				CurrentInputState.Mouse.Right.bDown = (Event == IE_Pressed);
-			}
+				FInputDeviceState InputState = CurrentMouseState;
+				InputState.InputDevice = EInputDevices::Mouse;
+				InputState.SetModifierKeyStates(
+					ViewportClient->IsShiftPressed(), ViewportClient->IsAltPressed(),
+					ViewportClient->IsCtrlPressed(), ViewportClient->IsCmdPressed());
 
-			InputRouter->PostInputEvent(InputState);
+				if (bIsLeftMouse)
+				{
+					InputState.Mouse.Left.SetStates(
+						(Event == IE_Pressed), (Event == IE_Pressed), (Event == IE_Released));
+					CurrentMouseState.Mouse.Left.bDown = (Event == IE_Pressed);
+				}
+				else if (bIsMiddleMouse)
+				{
+					InputState.Mouse.Middle.SetStates(
+						(Event == IE_Pressed), (Event == IE_Pressed), (Event == IE_Released));
+					CurrentMouseState.Mouse.Middle.bDown = (Event == IE_Pressed);
+				}
+				else
+				{
+					InputState.Mouse.Right.SetStates(
+						(Event == IE_Pressed), (Event == IE_Pressed), (Event == IE_Released));
+					CurrentMouseState.Mouse.Right.bDown = (Event == IE_Pressed);
+				}
 
-			if (InputRouter->HasActiveMouseCapture())
-			{
-				// what is this about? MeshPaintMode has it...
-				ViewportClient->bLockFlightCamera = true;
-				bHandled = true;   // indicate that we handled this event,
-								   // which will disable camera movement/etc ?
-			}
-			else
-			{
-				//ViewportClient->bLockFlightCamera = false;
-			}
+				InputRouter->PostInputEvent(InputState);
 
+				if (InputRouter->HasActiveMouseCapture())
+				{
+					// what is this about? MeshPaintMode has it...
+					ViewportClient->bLockFlightCamera = true;
+					bHandled = true;   // indicate that we handled this event,
+									   // which will disable camera movement/etc ?
+				}
+				else
+				{
+					//ViewportClient->bLockFlightCamera = false;
+				}
+
+			}
+		}
+		else if (Key.IsGamepadKey())
+		{
+			// not supported yet
+		}
+		else if (Key.IsTouch())
+		{
+			// not supported yet
+		}
+		else if (Key.IsFloatAxis() || Key.IsVectorAxis())
+		{
+			// not supported yet
+		}
+		else    // is this definitely a keyboard key?
+		{
+			FInputDeviceState InputState;
+			InputState.InputDevice = EInputDevices::Keyboard;
+			InputState.SetModifierKeyStates(
+				ViewportClient->IsShiftPressed(), ViewportClient->IsAltPressed(),
+				ViewportClient->IsCtrlPressed(), ViewportClient->IsCmdPressed());
+			InputState.Keyboard.ActiveKey.Button = Key;
+			bool bPressed = (Event == IE_Pressed);
+			InputState.Keyboard.ActiveKey.SetStates(bPressed, bPressed, !bPressed);
+			InputRouter->PostInputEvent(InputState);
 		}
 
 	}
@@ -548,8 +577,8 @@ bool UEdModeInteractiveToolsContext::MouseEnter(FEditorViewportClient* ViewportC
 	UE_LOG(LogTemp, Warning, TEXT("MOUSE ENTER"));
 #endif
 
-	CurrentInputState.Mouse.Position2D = FVector2D(x, y);
-	CurrentInputState.Mouse.WorldRay = GetRayFromMousePos(ViewportClient, Viewport, x, y);
+	CurrentMouseState.Mouse.Position2D = FVector2D(x, y);
+	CurrentMouseState.Mouse.WorldRay = GetRayFromMousePos(ViewportClient, Viewport, x, y);
 
 	return false;
 }
@@ -561,12 +590,12 @@ bool UEdModeInteractiveToolsContext::MouseMove(FEditorViewportClient* ViewportCl
 	//UE_LOG(LogTemp, Warning, TEXT("MOUSE MOVE"));
 #endif
 
-	CurrentInputState.Mouse.Position2D = FVector2D(x, y);
-	CurrentInputState.Mouse.WorldRay = GetRayFromMousePos(ViewportClient, Viewport, x, y);
-	FInputDeviceState InputState = CurrentInputState;
+	CurrentMouseState.Mouse.Position2D = FVector2D(x, y);
+	CurrentMouseState.Mouse.WorldRay = GetRayFromMousePos(ViewportClient, Viewport, x, y);
+	FInputDeviceState InputState = CurrentMouseState;
 	InputState.InputDevice = EInputDevices::Mouse;
 
-	InputState.SetKeyStates(
+	InputState.SetModifierKeyStates(
 		ViewportClient->IsShiftPressed(), ViewportClient->IsAltPressed(),
 		ViewportClient->IsCtrlPressed(), ViewportClient->IsCmdPressed());
 
@@ -608,18 +637,18 @@ bool UEdModeInteractiveToolsContext::CapturedMouseMove(FEditorViewportClient* In
 	//UE_LOG(LogTemp, Warning, TEXT("CAPTURED MOUSE MOVE"));
 #endif
 
-	FVector2D OldPosition = CurrentInputState.Mouse.Position2D;
-	CurrentInputState.Mouse.Position2D = FVector2D(InMouseX, InMouseY);
-	CurrentInputState.Mouse.WorldRay = GetRayFromMousePos(InViewportClient, InViewport, InMouseX, InMouseY);
+	FVector2D OldPosition = CurrentMouseState.Mouse.Position2D;
+	CurrentMouseState.Mouse.Position2D = FVector2D(InMouseX, InMouseY);
+	CurrentMouseState.Mouse.WorldRay = GetRayFromMousePos(InViewportClient, InViewport, InMouseX, InMouseY);
 
 	if (InputRouter->HasActiveMouseCapture())
 	{
-		FInputDeviceState InputState = CurrentInputState;
+		FInputDeviceState InputState = CurrentMouseState;
 		InputState.InputDevice = EInputDevices::Mouse;
-		InputState.SetKeyStates(
+		InputState.SetModifierKeyStates(
 			InViewportClient->IsShiftPressed(), InViewportClient->IsAltPressed(),
 			InViewportClient->IsCtrlPressed(), InViewportClient->IsCmdPressed());
-		InputState.Mouse.Delta2D = CurrentInputState.Mouse.Position2D - OldPosition;
+		InputState.Mouse.Delta2D = CurrentMouseState.Mouse.Position2D - OldPosition;
 		InputRouter->PostInputEvent(InputState);
 		return true;
 	}
diff --git a/Engine/Source/Editor/EditorInteractiveToolsFramework/Public/EdModeInteractiveToolsContext.h b/Engine/Source/Editor/EditorInteractiveToolsFramework/Public/EdModeInteractiveToolsContext.h
--- a/Engine/Source/Editor/EditorInteractiveToolsFramework/Public/EdModeInteractiveToolsContext.h
+++ b/Engine/Source/Editor/EditorInteractiveToolsFramework/Public/EdModeInteractiveToolsContext.h
@@ -71,7 +71,7 @@ protected:
 	bool bInvalidationPending;
 
 	/** Input event instance used to keep track of various button states, etc, that we cannot directly query on-demand */
-	FInputDeviceState CurrentInputState;
+	FInputDeviceState CurrentMouseState;
 
 	static FRay GetRayFromMousePos(FEditorViewportClient* ViewportClient, FViewport* Viewport, int MouseX, int MouseY);
 
diff --git a/Engine/Source/Runtime/InteractiveToolsFramework/Private/BaseBehaviors/KeyAsModifierInputBehavior.cpp b/Engine/Source/Runtime/InteractiveToolsFramework/Private/BaseBehaviors/KeyAsModifierInputBehavior.cpp
new file mode 100644
--- /dev/null
+++ b/Engine/Source/Runtime/InteractiveToolsFramework/Private/BaseBehaviors/KeyAsModifierInputBehavior.cpp
@@ -0,0 +1,56 @@
+// Copyright 1998-2019 Epic Games, Inc. All Rights Reserved.
+
+#include "BaseBehaviors/KeyAsModifierInputBehavior.h"
+
+UKeyAsModifierInputBehavior::UKeyAsModifierInputBehavior()
+{
+}
+
+void UKeyAsModifierInputBehavior::Initialize(IModifierToggleBehaviorTarget* TargetIn, int ModifierID, const FKey& ModifierKeyIn)
+{
+	this->Target = TargetIn;
+	ModifierKey = ModifierKeyIn;
+	FKey TempKey = ModifierKeyIn;
+	Modifiers.RegisterModifier(ModifierID, [TempKey](const FInputDeviceState& Input)
+	{
+		return Input.Keyboard.ActiveKey.Button == TempKey;
+	});
+}
+
+FInputCaptureRequest UKeyAsModifierInputBehavior::WantsCapture(const FInputDeviceState& Input)
+{
+	if ((ModifierCheckFunc == nullptr || ModifierCheckFunc(Input)))
+	{
+		if (Input.Keyboard.ActiveKey.Button == ModifierKey && Input.Keyboard.ActiveKey.bPressed)
+		{
+			return FInputCaptureRequest::Begin(this, EInputCaptureSide::Any);
+		}
+	}
+	return FInputCaptureRequest::Ignore();
+}
+
+
+FInputCaptureUpdate UKeyAsModifierInputBehavior::BeginCapture(const FInputDeviceState& Input, EInputCaptureSide Side)
+{
+	Modifiers.UpdateModifiers(Input, Target);
+	return FInputCaptureUpdate::Begin(this, EInputCaptureSide::Any);
+}
+
+
+FInputCaptureUpdate UKeyAsModifierInputBehavior::UpdateCapture(const FInputDeviceState& Input, const FInputCaptureData& Data)
+{
+	Modifiers.UpdateModifiers(Input, Target);
+
+	if (Input.Keyboard.ActiveKey.bReleased)
+	{
+		return FInputCaptureUpdate::End();
+	}
+
+	return FInputCaptureUpdate::Continue();
+}
+
+
+void UKeyAsModifierInputBehavior::ForceEndCapture(const FInputCaptureData& Data)
+{
+}
+
diff --git a/Engine/Source/Runtime/InteractiveToolsFramework/Private/InputRouter.cpp b/Engine/Source/Runtime/InteractiveToolsFramework/Private/InputRouter.cpp
--- a/Engine/Source/Runtime/InteractiveToolsFramework/Private/InputRouter.cpp
+++ b/Engine/Source/Runtime/InteractiveToolsFramework/Private/InputRouter.cpp
@@ -10,6 +10,8 @@ UInputRouter::UInputRouter()
 	ActiveLeftCaptureOwner = nullptr;
 	ActiveRightCapture = nullptr;
 	ActiveRightCaptureOwner = nullptr;
+	ActiveKeyboardCapture = nullptr;
+	ActiveKeyboardCaptureOwner = nullptr;
 
 	bAutoInvalidateOnHover = false;
 	bAutoInvalidateOnCapture = false;
@@ -45,25 +47,114 @@ void UInputRouter::DeregisterSource(IInputBehaviorSource* Source)
 
 
 
+void UInputRouter::PostInputEvent(const FInputDeviceState& Input)
+{
+	if (ActiveInputBehaviors->IsEmpty())
+	{
+		return;
+	}
+
+	if (Input.IsFromDevice(EInputDevices::Mouse))
+	{
+		PostInputEvent_Mouse(Input);
+	}
+	else if (Input.IsFromDevice(EInputDevices::Keyboard))
+	{
+		PostInputEvent_Keyboard(Input);
+	}
+	else
+	{
+		unimplemented();
+		TransactionsAPI->PostMessage(TEXT("UInteractiveToolManager::PostInputEvent - input device is not currently supported."), EToolMessageLevel::Internal);
+		return;
+	}
+}
 
 
 
+//
+// Keyboard event handling
+// 
 
-void UInputRouter::PostInputEvent(const FInputDeviceState& Input)
+
+void UInputRouter::PostInputEvent_Keyboard(const FInputDeviceState& Input)
 {
-	// currently mouse-specific...
+	if (ActiveKeyboardCapture != nullptr)
+	{
+		HandleCapturedKeyboardInput(Input);
+	}
+	else
+	{
+		ActiveKeyboardCaptureData = FInputCaptureData();
+		CheckForKeyboardCaptures(Input);
+	}
+}
 
-	if ( ActiveInputBehaviors->IsEmpty() )
+
+void UInputRouter::CheckForKeyboardCaptures(const FInputDeviceState& Input)
+{
+	TArray<FInputCaptureRequest> CaptureRequests;
+	ActiveInputBehaviors->CollectWantsCapture(Input, CaptureRequests);
+	if (CaptureRequests.Num() == 0)
 	{
 		return;
 	}
 
-	if (Input.IsFromDevice(EInputDevices::Mouse) == false)
+	CaptureRequests.StableSort();
+
+	bool bAccepted = false;
+	for (int i = 0; i < CaptureRequests.Num() && bAccepted == false; ++i)
+	{
+		FInputCaptureUpdate Result =
+			CaptureRequests[i].Source->BeginCapture(Input, EInputCaptureSide::Left);
+		if (Result.State == EInputCaptureState::Begin)
+		{
+			ActiveKeyboardCapture = Result.Source;
+			ActiveKeyboardCaptureOwner = CaptureRequests[i].Owner;
+			ActiveKeyboardCaptureData = Result.Data;
+			bAccepted = true;
+		}
+	}
+}
+
+
+void UInputRouter::HandleCapturedKeyboardInput(const FInputDeviceState& Input)
+{
+	if (ActiveKeyboardCapture == nullptr)
 	{
-		TransactionsAPI->PostMessage(TEXT("UInteractiveToolManager::PostInputEvent - non-mouse devices not currently supported."), EToolMessageLevel::Internal);
 		return;
 	}
 
+	FInputCaptureUpdate Result =
+		ActiveKeyboardCapture->UpdateCapture(Input, ActiveKeyboardCaptureData);
+
+	if (Result.State == EInputCaptureState::End)
+	{
+		ActiveKeyboardCapture = nullptr;
+		ActiveKeyboardCaptureOwner = nullptr;
+		ActiveKeyboardCaptureData = FInputCaptureData();
+	}
+	else if (Result.State != EInputCaptureState::Continue)
+	{
+		TransactionsAPI->PostMessage(TEXT("UInteractiveToolManager::HandleCapturedKeyboardInput - unexpected capture state!"), EToolMessageLevel::Internal);
+	}
+
+	if (bAutoInvalidateOnCapture)
+	{
+		TransactionsAPI->PostInvalidation();
+	}
+}
+
+
+
+//
+// Mouse event handling
+//
+
+
+
+void UInputRouter::PostInputEvent_Mouse(const FInputDeviceState& Input)
+{
 	if (ActiveLeftCapture != nullptr)
 	{
 		HandleCapturedMouseInput(Input);
@@ -170,6 +261,14 @@ void UInputRouter::HandleCapturedMouseInput(const FInputDeviceState& Input)
 
 void UInputRouter::ForceTerminateAll()
 {
+	if (ActiveKeyboardCapture != nullptr)
+	{
+		ActiveKeyboardCapture->ForceEndCapture(ActiveKeyboardCaptureData);
+		ActiveKeyboardCapture = nullptr;
+		ActiveKeyboardCaptureOwner = nullptr;
+		ActiveKeyboardCaptureData = FInputCaptureData();
+	}
+
 	if (ActiveLeftCapture != nullptr)
 	{
 		ActiveLeftCapture->ForceEndCapture(ActiveLeftCaptureData);
@@ -192,6 +291,14 @@ void UInputRouter::ForceTerminateAll()
 
 void UInputRouter::ForceTerminateSource(IInputBehaviorSource* Source)
 {
+	if (ActiveKeyboardCapture != nullptr && ActiveKeyboardCaptureOwner == Source)
+	{
+		ActiveKeyboardCapture->ForceEndCapture(ActiveKeyboardCaptureData);
+		ActiveKeyboardCapture = nullptr;
+		ActiveKeyboardCaptureOwner = nullptr;
+		ActiveKeyboardCaptureData = FInputCaptureData();
+	}
+
 	if (ActiveLeftCapture != nullptr && ActiveLeftCaptureOwner == Source)
 	{
 		ActiveLeftCapture->ForceEndCapture(ActiveLeftCaptureData);
diff --git a/Engine/Source/Runtime/InteractiveToolsFramework/Public/BaseBehaviors/KeyAsModifierInputBehavior.h b/Engine/Source/Runtime/InteractiveToolsFramework/Public/BaseBehaviors/KeyAsModifierInputBehavior.h
new file mode 100644
--- /dev/null
+++ b/Engine/Source/Runtime/InteractiveToolsFramework/Public/BaseBehaviors/KeyAsModifierInputBehavior.h
@@ -0,0 +1,61 @@
+// Copyright 1998-2019 Epic Games, Inc. All Rights Reserved.
+
+#pragma once
+
+#include "CoreMinimal.h"
+#include "InputBehavior.h"
+#include "InteractiveTool.h"
+#include "BaseBehaviors/BehaviorTargetInterfaces.h"
+#include "BaseBehaviors/InputBehaviorModifierStates.h"
+#include "KeyAsModifierInputBehavior.generated.h"
+
+
+/**
+ * UKeyAsModifierInputBehavior converts a specific key press/release into
+ * a "Modifier" toggle via the IModifierToggleBehaviorTarget interface
+ */
+UCLASS()
+class INTERACTIVETOOLSFRAMEWORK_API UKeyAsModifierInputBehavior : public UInputBehavior
+{
+	GENERATED_BODY()
+
+public:
+	UKeyAsModifierInputBehavior();
+
+	virtual EInputDevices GetSupportedDevices() override
+	{
+		return EInputDevices::Keyboard;
+	}
+
+	/**
+	 * Initialize this behavior with the given Target
+	 * @param Target implementor of modifier-toggle behavior
+	 * @param ModifierID integer ID that identifiers the modifier toggle
+	 * @param ModifierKey the key that will be used as the modifier toggle
+	 */
+	virtual void Initialize(IModifierToggleBehaviorTarget* Target, int ModifierID, const FKey& ModifierKey);
+
+	/**
+	 * WantsCapture() will only return capture request if this function returns true (or is null)
+	 * Intended to be used for alt/ctrl/cmd/shift modifiers on the main ModifierKey
+	 */
+	TFunction<bool(const FInputDeviceState&)> ModifierCheckFunc = nullptr;
+
+	// UInputBehavior implementation
+
+	virtual FInputCaptureRequest WantsCapture(const FInputDeviceState& Input) override;
+	virtual FInputCaptureUpdate BeginCapture(const FInputDeviceState& Input, EInputCaptureSide Side) override;
+	virtual FInputCaptureUpdate UpdateCapture(const FInputDeviceState& Input, const FInputCaptureData& Data) override;
+	virtual void ForceEndCapture(const FInputCaptureData& Data) override;
+
+
+protected:
+	/** Modifier Target object */
+	IModifierToggleBehaviorTarget* Target;
+
+	/** Key that is used as modifier */
+	FKey ModifierKey;
+
+	/** Modifier set for this behavior, internally initialized with check on ModifierKey */
+	FInputBehaviorModifierStates Modifiers;
+};
diff --git a/Engine/Source/Runtime/InteractiveToolsFramework/Public/InputRouter.h b/Engine/Source/Runtime/InteractiveToolsFramework/Public/InputRouter.h
--- a/Engine/Source/Runtime/InteractiveToolsFramework/Public/InputRouter.h
+++ b/Engine/Source/Runtime/InteractiveToolsFramework/Public/InputRouter.h
@@ -81,6 +81,10 @@ protected:
 	UPROPERTY()
 	UInputBehaviorSet* ActiveInputBehaviors;
 
+	UInputBehavior* ActiveKeyboardCapture;
+	void* ActiveKeyboardCaptureOwner;
+	FInputCaptureData ActiveKeyboardCaptureData;
+
 	UInputBehavior* ActiveLeftCapture;
 	void* ActiveLeftCaptureOwner;
 	FInputCaptureData ActiveLeftCaptureData;
@@ -91,6 +95,11 @@ protected:
 
 	FInputDeviceState LastHoverInput;
 
+	virtual void PostInputEvent_Keyboard(const FInputDeviceState& Input);
+	void CheckForKeyboardCaptures(const FInputDeviceState& Input);
+	void HandleCapturedKeyboardInput(const FInputDeviceState& Input);
+
+	virtual void PostInputEvent_Mouse(const FInputDeviceState& Input);
 	void CheckForMouseCaptures(const FInputDeviceState& Input);
 	void HandleCapturedMouseInput(const FInputDeviceState& Input);
 };
\ No newline at end of file
diff --git a/Engine/Source/Runtime/InteractiveToolsFramework/Public/InputState.h b/Engine/Source/Runtime/InteractiveToolsFramework/Public/InputState.h
--- a/Engine/Source/Runtime/InteractiveToolsFramework/Public/InputState.h
+++ b/Engine/Source/Runtime/InteractiveToolsFramework/Public/InputState.h
@@ -15,11 +15,12 @@ UENUM()
 enum class EInputDevices
 {
 	None = 0,
-	Mouse = 1,
-	Gamepad = 2,
+	Keyboard = 1,
+	Mouse = 2,
+	Gamepad = 4,
 
-	OculusTouch = 4,
-	HTCViveWands = 8,
+	OculusTouch = 8,
+	HTCViveWands = 16,
 	AnySpatialDevice = OculusTouch | HTCViveWands,
 
 	TabletFingers = 1024
@@ -71,6 +72,9 @@ struct FInputDeviceRay
  */
 struct FDeviceButtonState
 {
+	/** Button identifier */
+	FKey Button;
+
 	/** Was the button pressed down this frame. This should happen once per "click" */
 	bool bPressed;
 	/** Is the button currently pressed down. This should be true every frame the button is pressed. */
@@ -80,6 +84,13 @@ struct FDeviceButtonState
 
 	FDeviceButtonState()
 	{
+		Button = FKey();
+		bPressed = bDown = bReleased = false;
+	}
+
+	FDeviceButtonState(const FKey& ButtonIn)
+	{
+		Button = ButtonIn;
 		bPressed = bDown = bReleased = false;
 	}
 
@@ -94,6 +105,18 @@ struct FDeviceButtonState
 
 
 
+/**
+ * Current state of active keyboard key at a point in time
+ * @todo would be useful to track set of active keys
+ */
+struct FKeyboardInputDeviceState
+{
+	/** state of active key that was modified (ie press or release) */
+	FDeviceButtonState ActiveKey;
+};
+
+
+
 /**
  * Current State of a physical Mouse device at a point in time.
  */
@@ -121,9 +144,9 @@ struct FMouseInputDeviceState
 
 	FMouseInputDeviceState()
 	{
-		Left = FDeviceButtonState();
-		Middle = FDeviceButtonState();
-		Right = FDeviceButtonState();
+		Left = FDeviceButtonState(EKeys::LeftMouseButton);
+		Middle = FDeviceButtonState(EKeys::MiddleMouseButton);
+		Right = FDeviceButtonState(EKeys::RightMouseButton);
 		WheelDelta = false;
 		Position2D = FVector2D::ZeroVector;
 		Delta2D = FVector2D::ZeroVector;
@@ -156,6 +179,10 @@ struct FInputDeviceState
 	/** Is they keyboard CMD modifier key currently pressed down (only on Apple devices) */
 	bool bCmdKeyDown;
 
+
+	/** Current state of Keyboard device, if InputDevice == EInputDevices::Keyboard */
+	FKeyboardInputDeviceState Keyboard;
+
 	/** Current state of Mouse device, if InputDevice == EInputDevices::Mouse */
 	FMouseInputDeviceState Mouse;
 
@@ -164,11 +191,12 @@ struct FInputDeviceState
 	{
 		InputDevice = EInputDevices::None;
 		bShiftKeyDown = bAltKeyDown = bCtrlKeyDown = bCmdKeyDown = false;
+		Keyboard = FKeyboardInputDeviceState();
 		Mouse = FMouseInputDeviceState();
 	}
 
 	/** Update keyboard modifier key states */
-	void SetKeyStates(bool bShiftDown, bool bAltDown, bool bCtrlDown, bool bCmdDown) 
+	void SetModifierKeyStates(bool bShiftDown, bool bAltDown, bool bCtrlDown, bool bCmdDown) 
 	{
 		bShiftKeyDown = bShiftDown;
 		bAltKeyDown = bAltDown;
