shutdown active tool on PIE and save world (default to accept)
#rnx

diff --git a/Engine/Source/Editor/EditorInteractiveToolsFramework/Private/EdModeInteractiveToolsContext.cpp b/Engine/Source/Editor/EditorInteractiveToolsFramework/Private/EdModeInteractiveToolsContext.cpp
--- a/Engine/Source/Editor/EditorInteractiveToolsFramework/Private/EdModeInteractiveToolsContext.cpp
+++ b/Engine/Source/Editor/EditorInteractiveToolsFramework/Private/EdModeInteractiveToolsContext.cpp
@@ -2,6 +2,7 @@
 
 
 #include "EdModeInteractiveToolsContext.h"
+#include "Editor.h"
 #include "EditorViewportClient.h"
 #include "EditorModeManager.h"
 #include "LevelEditorViewport.h"   // for GCurrentLevelEditingViewportClient
@@ -315,11 +316,23 @@ void UEdModeInteractiveToolsContext::Initialize(IToolsContextQueriesAPI* Queries
 {
 	UInteractiveToolsContext::Initialize(QueriesAPIIn, TransactionsAPIIn);
 
+	BeginPIEDelegateHandle = FEditorDelegates::BeginPIE.AddLambda([this](bool bSimulating)
+	{
+		TerminateActiveToolsOnPIEStart();
+	});
+	PreSaveWorldDelegateHandle = FEditorDelegates::PreSaveWorld.AddLambda([this](uint32 SaveFlags, UWorld* World)
+	{
+		TerminateActiveToolsOnSaveWorld();
+	});
+
 	bInvalidationPending = false;
 }
 
 void UEdModeInteractiveToolsContext::Shutdown()
 {
+	FEditorDelegates::BeginPIE.Remove(BeginPIEDelegateHandle);
+	FEditorDelegates::PreSaveWorld.Remove(PreSaveWorldDelegateHandle);
+
 	UInteractiveToolsContext::Shutdown();
 }
 
@@ -380,6 +393,42 @@ void UEdModeInteractiveToolsContext::ShutdownContext()
 
 
 
+
+void UEdModeInteractiveToolsContext::TerminateActiveToolsOnPIEStart()
+{
+	if (ToolManager->HasActiveTool(EToolSide::Left))
+	{
+		EToolShutdownType ShutdownType = ToolManager->CanAcceptActiveTool(EToolSide::Left) ?
+			EToolShutdownType::Accept : EToolShutdownType::Cancel;
+		ToolManager->DeactivateTool(EToolSide::Left, ShutdownType);
+	}
+	if (ToolManager->HasActiveTool(EToolSide::Right))
+	{
+		EToolShutdownType ShutdownType = ToolManager->CanAcceptActiveTool(EToolSide::Right) ?
+			EToolShutdownType::Accept : EToolShutdownType::Cancel;
+		ToolManager->DeactivateTool(EToolSide::Right, ShutdownType);
+	}
+}
+void UEdModeInteractiveToolsContext::TerminateActiveToolsOnSaveWorld()
+{
+	if (ToolManager->HasActiveTool(EToolSide::Left))
+	{
+		EToolShutdownType ShutdownType = ToolManager->CanAcceptActiveTool(EToolSide::Left) ?
+			EToolShutdownType::Accept : EToolShutdownType::Cancel;
+		ToolManager->DeactivateTool(EToolSide::Left, ShutdownType);
+	}
+	if (ToolManager->HasActiveTool(EToolSide::Right))
+	{
+		EToolShutdownType ShutdownType = ToolManager->CanAcceptActiveTool(EToolSide::Right) ?
+			EToolShutdownType::Accept : EToolShutdownType::Cancel;
+		ToolManager->DeactivateTool(EToolSide::Right, ShutdownType);
+	}
+}
+
+
+
+
+
 void UEdModeInteractiveToolsContext::PostInvalidation()
 {
 	bInvalidationPending = true;
diff --git a/Engine/Source/Editor/EditorInteractiveToolsFramework/Public/EdModeInteractiveToolsContext.h b/Engine/Source/Editor/EditorInteractiveToolsFramework/Public/EdModeInteractiveToolsContext.h
--- a/Engine/Source/Editor/EditorInteractiveToolsFramework/Public/EdModeInteractiveToolsContext.h
+++ b/Engine/Source/Editor/EditorInteractiveToolsFramework/Public/EdModeInteractiveToolsContext.h
@@ -25,6 +25,11 @@ public:
 	virtual void InitializeContextFromEdMode(FEdMode* EditorMode);
 	virtual void ShutdownContext();
 
+	// default behavior is to accept active tool
+	virtual void TerminateActiveToolsOnPIEStart();
+
+	// default behavior is to accept active tool
+	virtual void TerminateActiveToolsOnSaveWorld();
 
 	IToolsContextQueriesAPI* GetQueriesAPI() const { return QueriesAPI; }
 	IToolsContextTransactionsAPI* GetTransactionAPI() const { return TransactionAPI; }
@@ -63,6 +68,9 @@ public:
 protected:
 	FEdMode* EditorMode;
 
+	FDelegateHandle BeginPIEDelegateHandle;
+	FDelegateHandle PreSaveWorldDelegateHandle;
+
 	IToolsContextQueriesAPI* QueriesAPI;
 	IToolsContextTransactionsAPI* TransactionAPI;
 	IToolsContextAssetAPI* AssetAPI;
