enable RF_Transactional flag on UStaticMesh before making changes, so that the mesh updates are transacted
#rnx

diff --git a/Engine/Source/Editor/EditorInteractiveToolsFramework/Private/EditorComponentSourceFactory.cpp b/Engine/Source/Editor/EditorInteractiveToolsFramework/Private/EditorComponentSourceFactory.cpp
--- a/Engine/Source/Editor/EditorInteractiveToolsFramework/Private/EditorComponentSourceFactory.cpp
+++ b/Engine/Source/Editor/EditorInteractiveToolsFramework/Private/EditorComponentSourceFactory.cpp
@@ -73,7 +73,12 @@ bool FStaticMeshComponentMeshDescriptionSource::IsReadOnly() const
 void FStaticMeshComponentMeshDescriptionSource::CommitInPlaceModification(const TFunction<void(FMeshDescription*)>& ModifyFunction) 
 {
 	UStaticMesh* StaticMesh = Component->GetStaticMesh();
-	StaticMesh->Modify();
+
+	// make sure transactional flag is on
+	StaticMesh->SetFlags(RF_Transactional);
+
+	bool bSavedToTransactionBuffer = StaticMesh->Modify();
+	check(bSavedToTransactionBuffer);
 	FMeshDescription* MeshDescription = StaticMesh->GetMeshDescription(LODIndex);
 
 	ModifyFunction(MeshDescription);
