Don't forward captured-mouse-move events to ToolsFramework if alt is down (ie camera change).
Minor API changes to InteractiveGizmoManager and MeshTangents.
#rnx

diff --git a/Engine/Plugins/Experimental/GeometryProcessing/Source/DynamicMesh/Public/MeshTangents.h b/Engine/Plugins/Experimental/GeometryProcessing/Source/DynamicMesh/Public/MeshTangents.h
--- a/Engine/Plugins/Experimental/GeometryProcessing/Source/DynamicMesh/Public/MeshTangents.h
+++ b/Engine/Plugins/Experimental/GeometryProcessing/Source/DynamicMesh/Public/MeshTangents.h
@@ -46,7 +46,7 @@ public:
 	 * @param TriangleID triangle index in mesh
 	 * @param TriVertIdx vertex index in range 0,1,2
 	 */
-	void GetPerTriangleTangent(int TriangleID, int TriVertIdx, FVector3<RealType>& TangentOut, FVector3<RealType>& BitangentOut)
+	void GetPerTriangleTangent(int TriangleID, int TriVertIdx, FVector3<RealType>& TangentOut, FVector3<RealType>& BitangentOut) const
 	{
 		int k = TriangleID * 3 + TriVertIdx;
 		TangentOut = Tangents[k];
diff --git a/Engine/Source/Editor/EditorInteractiveToolsFramework/Private/EdModeInteractiveToolsContext.cpp b/Engine/Source/Editor/EditorInteractiveToolsFramework/Private/EdModeInteractiveToolsContext.cpp
--- a/Engine/Source/Editor/EditorInteractiveToolsFramework/Private/EdModeInteractiveToolsContext.cpp
+++ b/Engine/Source/Editor/EditorInteractiveToolsFramework/Private/EdModeInteractiveToolsContext.cpp
@@ -527,7 +527,11 @@ bool UEdModeInteractiveToolsContext::InputKey(FEditorViewportClient* ViewportCli
 		}
 	}
 
-
+	// if alt is down we do not process mouse event
+	if (ViewportClient->IsAltPressed())
+	{
+		return false;
+	}
 
 	if (Event == IE_Pressed || Event == IE_Released)
 	{
@@ -686,6 +690,12 @@ bool UEdModeInteractiveToolsContext::CapturedMouseMove(FEditorViewportClient* In
 	//UE_LOG(LogTemp, Warning, TEXT("CAPTURED MOUSE MOVE"));
 #endif
 
+	// if alt is down we will not allow client to see this event
+	if (InViewportClient->IsAltPressed())
+	{
+		return false;
+	}
+
 	FVector2D OldPosition = CurrentMouseState.Mouse.Position2D;
 	CurrentMouseState.Mouse.Position2D = FVector2D(InMouseX, InMouseY);
 	CurrentMouseState.Mouse.WorldRay = GetRayFromMousePos(InViewportClient, InViewport, InMouseX, InMouseY);
diff --git a/Engine/Source/Runtime/InteractiveToolsFramework/Public/InteractiveGizmoManager.h b/Engine/Source/Runtime/InteractiveToolsFramework/Public/InteractiveGizmoManager.h
--- a/Engine/Source/Runtime/InteractiveToolsFramework/Public/InteractiveGizmoManager.h
+++ b/Engine/Source/Runtime/InteractiveToolsFramework/Public/InteractiveGizmoManager.h
@@ -149,6 +149,15 @@ public:
 	virtual void Render(IToolsContextRenderAPI* RenderAPI);
 
 
+
+	//
+	// access to APIs, etc
+	//
+
+	/** @return current IToolsContextQueriesAPI */
+	virtual IToolsContextQueriesAPI* GetContextQueriesAPI() { return QueriesAPI; }
+
+
 public:
 	/** set of Currently-active Gizmos */
 	UPROPERTY()
