add helper functions to make it easier to connect editor UI to ToolManager
#rnx

diff --git a/Engine/Source/Editor/EditorInteractiveToolsFramework/Private/EdModeInteractiveToolsContext.cpp b/Engine/Source/Editor/EditorInteractiveToolsFramework/Private/EdModeInteractiveToolsContext.cpp
--- a/Engine/Source/Editor/EditorInteractiveToolsFramework/Private/EdModeInteractiveToolsContext.cpp
+++ b/Engine/Source/Editor/EditorInteractiveToolsFramework/Private/EdModeInteractiveToolsContext.cpp
@@ -742,3 +742,47 @@ FRay UEdModeInteractiveToolsContext::GetRayFromMousePos(FEditorViewportClient* V
 
 	return FRay(MouseViewportRay.GetOrigin(), MouseViewportRay.GetDirection(), true);
 }
+
+
+
+
+
+
+bool UEdModeInteractiveToolsContext::CanStartTool(const FString& ToolTypeIdentifier) const
+{
+	return (ToolManager->HasActiveTool(EToolSide::Mouse) == false) &&
+		(ToolManager->CanActivateTool(EToolSide::Mouse, ToolTypeIdentifier) == true);
+}
+bool UEdModeInteractiveToolsContext::ActiveToolHasAccept() const
+{
+	return ToolManager->HasActiveTool(EToolSide::Mouse) &&
+		ToolManager->GetActiveTool(EToolSide::Mouse)->HasAccept();
+}
+bool UEdModeInteractiveToolsContext::CanAcceptActiveTool() const
+{
+	return ToolManager->CanAcceptActiveTool(EToolSide::Mouse);
+}
+bool UEdModeInteractiveToolsContext::CanCancelActiveTool() const
+{
+	return ToolManager->CanCancelActiveTool(EToolSide::Mouse);
+}
+
+bool UEdModeInteractiveToolsContext::CanCompleteActiveTool() const
+{
+	return ToolManager->HasActiveTool(EToolSide::Mouse) && CanCancelActiveTool() == false;
+}
+void UEdModeInteractiveToolsContext::StartTool(const FString& ToolTypeIdentifier)
+{
+	if (ToolManager->SelectActiveToolType(EToolSide::Mouse, ToolTypeIdentifier) == false)
+	{
+		UE_LOG(LogTemp, Warning, TEXT("ToolManager: Unknown Tool Type %s"), *ToolTypeIdentifier);
+	}
+	else
+	{
+		ToolManager->ActivateTool(EToolSide::Mouse);
+	}
+}
+void UEdModeInteractiveToolsContext::EndTool(EToolShutdownType ShutdownType)
+{
+	ToolManager->DeactivateTool(EToolSide::Mouse, ShutdownType);
+}
\ No newline at end of file
diff --git a/Engine/Source/Editor/EditorInteractiveToolsFramework/Public/EdModeInteractiveToolsContext.h b/Engine/Source/Editor/EditorInteractiveToolsFramework/Public/EdModeInteractiveToolsContext.h
--- a/Engine/Source/Editor/EditorInteractiveToolsFramework/Public/EdModeInteractiveToolsContext.h
+++ b/Engine/Source/Editor/EditorInteractiveToolsFramework/Public/EdModeInteractiveToolsContext.h
@@ -55,6 +55,19 @@ public:
 	virtual bool EndTracking(FEditorViewportClient* InViewportClient, FViewport* InViewport);
 
 
+	//
+	// Utility functions useful for hooking up to UICommand/etc
+	//
+
+	virtual bool CanStartTool(const FString& ToolTypeIdentifier) const;
+	virtual bool ActiveToolHasAccept() const;
+	virtual bool CanAcceptActiveTool() const;
+	virtual bool CanCancelActiveTool() const;
+	virtual bool CanCompleteActiveTool() const;
+	virtual void StartTool(const FString& ToolTypeIdentifier);
+	virtual void EndTool(EToolShutdownType ShutdownType);
+
+
 protected:
 	// we hide these 
 	virtual void Initialize(IToolsContextQueriesAPI* QueriesAPI, IToolsContextTransactionsAPI* TransactionsAPI) override;
