rebuild physics state after updating mesh. Works going forward, but state rebuild does not similarly happen on undo.
#rnx

diff --git a/Engine/Source/Editor/EditorInteractiveToolsFramework/Private/EditorComponentSourceFactory.cpp b/Engine/Source/Editor/EditorInteractiveToolsFramework/Private/EditorComponentSourceFactory.cpp
--- a/Engine/Source/Editor/EditorInteractiveToolsFramework/Private/EditorComponentSourceFactory.cpp
+++ b/Engine/Source/Editor/EditorInteractiveToolsFramework/Private/EditorComponentSourceFactory.cpp
@@ -5,8 +5,7 @@
 #include "Engine/StaticMesh.h"
 #include "Components/ActorComponent.h"
 #include "Components/StaticMeshComponent.h"
-
-
+#include "PhysicsEngine/BodySetup.h"
 
 
 TUniquePtr<IMeshDescriptionSource> FEditorComponentSourceFactory::MakeMeshDescriptionSource(UActorComponent* Component)
@@ -74,6 +73,8 @@ bool FStaticMeshComponentMeshDescriptionSource::IsReadOnly() const
 
 void FStaticMeshComponentMeshDescriptionSource::CommitInPlaceModification(const TFunction<void(FMeshDescription*)>& ModifyFunction) 
 {
+	//bool bSaved = Component->Modify();
+	//check(bSaved);
 	UStaticMesh* StaticMesh = Component->GetStaticMesh();
 
 	// make sure transactional flag is on
@@ -87,6 +88,11 @@ void FStaticMeshComponentMeshDescriptionSource::CommitInPlaceModification(const
 
 	StaticMesh->CommitMeshDescription(LODIndex);
 	StaticMesh->PostEditChange();
+
+	// this rebuilds physics, but it doesn't undo!
+	Component->RecreatePhysicsState();
+
+	//Component->PostEditChange();
 }
 
 
